Updated some code to work with Python 3 API changes

diff --git a/Engine/Plugins/Experimental/PythonScriptPlugin/Source/PythonScriptPlugin/Private/PyWrapperMath.cpp b/Engine/Plugins/Experimental/PythonScriptPlugin/Source/PythonScriptPlugin/Private/PyWrapperMath.cpp
--- a/Engine/Plugins/Experimental/PythonScriptPlugin/Source/PythonScriptPlugin/Private/PyWrapperMath.cpp
+++ b/Engine/Plugins/Experimental/PythonScriptPlugin/Source/PythonScriptPlugin/Private/PyWrapperMath.cpp
@@ -810,8 +810,13 @@ PyTypeObject InitializePyWrapperMathType_Vector()
 	PyNumber.nb_inplace_subtract = (binaryfunc)&EvaulateOperatorStack_Inline_Struct_ReturnStruct<FPyWrapperVector, FSubOperation>;
 	PyNumber.nb_multiply = (binaryfunc)&EvaulateOperatorStack_StructOrFloat_ReturnStruct<FPyWrapperVector, FMulOperation>;
 	PyNumber.nb_inplace_multiply = (binaryfunc)&EvaulateOperatorStack_Inline_StructOrFloat_ReturnStruct<FPyWrapperVector, FMulOperation>;
+#if PY_MAJOR_VERSION >= 3
+	PyNumber.nb_true_divide = (binaryfunc)&EvaulateOperatorStack_StructOrFloat_ReturnStruct<FPyWrapperVector, FDivOperation>;
+	PyNumber.nb_inplace_true_divide = (binaryfunc)&EvaulateOperatorStack_Inline_StructOrFloat_ReturnStruct<FPyWrapperVector, FDivOperation>;
+#else	// PY_MAJOR_VERSION >= 3
 	PyNumber.nb_divide = (binaryfunc)&EvaulateOperatorStack_StructOrFloat_ReturnStruct<FPyWrapperVector, FDivOperation>;
 	PyNumber.nb_inplace_divide = (binaryfunc)&EvaulateOperatorStack_Inline_StructOrFloat_ReturnStruct<FPyWrapperVector, FDivOperation>;
+#endif	// PY_MAJOR_VERSION >= 3
 	PyNumber.nb_or = (binaryfunc)&EvaulateOperatorStack_Struct_ReturnFloat<FPyWrapperVector, FBitwiseOrOperation>;
 	PyNumber.nb_xor = (binaryfunc)&EvaulateOperatorStack_Struct_ReturnStruct<FPyWrapperVector, FBitwiseXorOperation>;
 
@@ -1005,8 +1010,13 @@ PyTypeObject InitializePyWrapperMathType_Vector2D()
 	PyNumber.nb_inplace_subtract = (binaryfunc)&EvaulateOperatorStack_Inline_Struct_ReturnStruct<FPyWrapperVector2D, FSubOperation>;
 	PyNumber.nb_multiply = (binaryfunc)&EvaulateOperatorStack_StructOrFloat_ReturnStruct<FPyWrapperVector2D, FMulOperation>;
 	PyNumber.nb_inplace_multiply = (binaryfunc)&EvaulateOperatorStack_Inline_StructOrFloat_ReturnStruct<FPyWrapperVector2D, FMulOperation>;
+#if PY_MAJOR_VERSION >= 3
+	PyNumber.nb_true_divide = (binaryfunc)&EvaulateOperatorStack_StructOrFloat_ReturnStruct<FPyWrapperVector2D, FDivOperation>;
+	PyNumber.nb_inplace_true_divide = (binaryfunc)&EvaulateOperatorStack_Inline_StructOrFloat_ReturnStruct<FPyWrapperVector2D, FDivOperation>;
+#else	// PY_MAJOR_VERSION >= 3
 	PyNumber.nb_divide = (binaryfunc)&EvaulateOperatorStack_StructOrFloat_ReturnStruct<FPyWrapperVector2D, FDivOperation>;
 	PyNumber.nb_inplace_divide = (binaryfunc)&EvaulateOperatorStack_Inline_StructOrFloat_ReturnStruct<FPyWrapperVector2D, FDivOperation>;
+#endif	// PY_MAJOR_VERSION >= 3
 	PyNumber.nb_or = (binaryfunc)&EvaulateOperatorStack_Struct_ReturnFloat<FPyWrapperVector2D, FBitwiseOrOperation>;
 	PyNumber.nb_xor = (binaryfunc)&EvaulateOperatorStack_Struct_ReturnFloat<FPyWrapperVector2D, FBitwiseXorOperation>;
 
@@ -1162,8 +1172,13 @@ PyTypeObject InitializePyWrapperMathType_Vector4()
 	PyNumber.nb_inplace_subtract = (binaryfunc)&EvaulateOperatorStack_Inline_Struct_ReturnStruct<FPyWrapperVector4, FSubOperation>;
 	PyNumber.nb_multiply = (binaryfunc)&EvaulateOperatorStack_StructOrFloat_ReturnStruct<FPyWrapperVector4, FMulOperation>;
 	PyNumber.nb_inplace_multiply = (binaryfunc)&EvaulateOperatorStack_Inline_StructOrFloat_ReturnStruct<FPyWrapperVector4, FMulOperation>;
+#if PY_MAJOR_VERSION >= 3
+	PyNumber.nb_true_divide = (binaryfunc)&EvaulateOperatorStack_StructOrFloat_ReturnStruct<FPyWrapperVector4, FDivOperation>;
+	PyNumber.nb_inplace_true_divide = (binaryfunc)&EvaulateOperatorStack_Inline_Struct_ReturnStruct<FPyWrapperVector4, FDivOperation>;
+#else	// PY_MAJOR_VERSION >= 3
 	PyNumber.nb_divide = (binaryfunc)&EvaulateOperatorStack_StructOrFloat_ReturnStruct<FPyWrapperVector4, FDivOperation>;
 	PyNumber.nb_inplace_divide = (binaryfunc)&EvaulateOperatorStack_Inline_Struct_ReturnStruct<FPyWrapperVector4, FDivOperation>;
+#endif	// PY_MAJOR_VERSION >= 3
 	//PyNumber.nb_or = (binaryfunc)&EvaulateOperatorStack_Struct_ReturnStruct<FPyWrapperVector2D, FBitwiseOrOperation>;
 	PyNumber.nb_xor = (binaryfunc)&EvaulateOperatorStack_Struct_ReturnStruct<FPyWrapperVector4, FBitwiseXorOperation>;
 
@@ -1395,8 +1410,13 @@ PyTypeObject InitializePyWrapperMathType_Quat()
 	PyNumber.nb_inplace_subtract = (binaryfunc)&EvaulateOperatorStack_Inline_Struct_ReturnStruct<FPyWrapperQuat, FSubOperation>;
 	PyNumber.nb_multiply = (binaryfunc)&EvaulateOperatorStack_StructOrFloat_ReturnStruct<FPyWrapperQuat, FMulOperation>;
 	PyNumber.nb_inplace_multiply = (binaryfunc)&EvaulateOperatorStack_Inline_Struct_ReturnStruct<FPyWrapperQuat, FMulOperation>;
+#if PY_MAJOR_VERSION >= 3
+	PyNumber.nb_true_divide = (binaryfunc)&EvaulateOperatorStack_Float_ReturnStruct<FPyWrapperQuat, FDivOperation>;
+	PyNumber.nb_inplace_true_divide = (binaryfunc)&EvaulateOperatorStack_Inline_Float_ReturnStruct<FPyWrapperQuat, FDivOperation>;
+#else	// PY_MAJOR_VERSION >= 3
 	PyNumber.nb_divide = (binaryfunc)&EvaulateOperatorStack_Float_ReturnStruct<FPyWrapperQuat, FDivOperation>;
 	PyNumber.nb_inplace_divide = (binaryfunc)&EvaulateOperatorStack_Inline_Float_ReturnStruct<FPyWrapperQuat, FDivOperation>;
+#endif	// PY_MAJOR_VERSION >= 3
 
 	PyType.tp_methods = PyMethods;
 	PyType.tp_getset = PyGetSets;
@@ -1651,8 +1671,13 @@ PyTypeObject InitializePyWrapperMathType_LinearColor()
 	PyNumber.nb_inplace_subtract = (binaryfunc)&EvaulateOperatorStack_Inline_Struct_ReturnStruct<FPyWrapperLinearColor, FSubOperation>;
 	PyNumber.nb_multiply = (binaryfunc)&EvaulateOperatorStack_StructOrFloat_ReturnStruct<FPyWrapperLinearColor, FMulOperation>;
 	PyNumber.nb_inplace_multiply = (binaryfunc)&EvaulateOperatorStack_Inline_Struct_ReturnStruct<FPyWrapperLinearColor, FMulOperation>;
+#if PY_MAJOR_VERSION >= 3
+	PyNumber.nb_true_divide = (binaryfunc)&EvaulateOperatorStack_StructOrFloat_ReturnStruct<FPyWrapperLinearColor, FDivOperation>;
+	PyNumber.nb_inplace_true_divide = (binaryfunc)&EvaulateOperatorStack_Inline_StructOrFloat_ReturnStruct<FPyWrapperLinearColor, FDivOperation>;
+#else	// PY_MAJOR_VERSION >= 3
 	PyNumber.nb_divide = (binaryfunc)&EvaulateOperatorStack_StructOrFloat_ReturnStruct<FPyWrapperLinearColor, FDivOperation>;
 	PyNumber.nb_inplace_divide = (binaryfunc)&EvaulateOperatorStack_Inline_StructOrFloat_ReturnStruct<FPyWrapperLinearColor, FDivOperation>;
+#endif	// PY_MAJOR_VERSION >= 3
 
 	PyType.tp_methods = PyMethods;
 	PyType.tp_getset = PyGetSets;
diff --git a/Engine/Plugins/Experimental/PythonScriptPlugin/Source/PythonScriptPlugin/Private/PythonScriptPlugin.cpp b/Engine/Plugins/Experimental/PythonScriptPlugin/Source/PythonScriptPlugin/Private/PythonScriptPlugin.cpp
--- a/Engine/Plugins/Experimental/PythonScriptPlugin/Source/PythonScriptPlugin/Private/PythonScriptPlugin.cpp
+++ b/Engine/Plugins/Experimental/PythonScriptPlugin/Source/PythonScriptPlugin/Private/PythonScriptPlugin.cpp
@@ -572,8 +572,8 @@ PyObject* FPythonScriptPlugin::EvalString(const TCHAR* InStr, const TCHAR* InCon
 
 PyObject* FPythonScriptPlugin::EvalString(const TCHAR* InStr, const TCHAR* InContext, const int InMode, PyObject* InGlobalDict, PyObject* InLocalDict)
 {
-	FScopedSlowTask SlowTask(0.f, (LOCTEXT("PythonScript_EvalString", "Running Python")));
-	SlowTask.MakeDialog();
+	FScopedSlowTask SlowTask(0.f, LOCTEXT("PythonScript_EvalString", "Running Python"));
+	SlowTask.MakeDialogDelayed(1.0f);
 
 	PyCompilerFlags *PyCompFlags = nullptr;
 
