Copying //UE4/Dev-Enterprise to //UE4/Dev-Main (Source: //UE4/Dev-Enterprise @ 3972172)

diff --git a/Engine/Plugins/Experimental/PythonScriptPlugin/Source/PythonScriptPlugin/Private/PyConversion.cpp b/Engine/Plugins/Experimental/PythonScriptPlugin/Source/PythonScriptPlugin/Private/PyConversion.cpp
--- a/Engine/Plugins/Experimental/PythonScriptPlugin/Source/PythonScriptPlugin/Private/PyConversion.cpp
+++ b/Engine/Plugins/Experimental/PythonScriptPlugin/Source/PythonScriptPlugin/Private/PyConversion.cpp
@@ -25,12 +25,12 @@
 
 #define PYCONVERSION_RETURN(RESULT, ERROR_CTX, ERROR_MSG)									\
 	{																						\
-		const bool bPyConversionReturnResult_Internal = (RESULT);							\
-		if (!bPyConversionReturnResult_Internal && SetErrorState == ESetErrorState::Yes)	\
+		const FPyConversionResult PyConversionReturnResult_Internal = (RESULT);				\
+		if (!PyConversionReturnResult_Internal && SetErrorState == ESetErrorState::Yes)		\
 		{																					\
 			PyUtil::SetPythonError(PyExc_TypeError, (ERROR_CTX), (ERROR_MSG));				\
 		}																					\
-		return bPyConversionReturnResult_Internal;											\
+		return PyConversionReturnResult_Internal;											\
 	}
 
 namespace PyConversion
@@ -39,79 +39,124 @@ namespace PyConversion
 namespace Internal
 {
 
-bool NativizeStruct(PyObject* PyObj, UScriptStruct* StructType, void* StructInstance, const ESetErrorState SetErrorState)
+FPyConversionResult NativizeStruct(PyObject* PyObj, UScriptStruct* StructType, void* StructInstance, const ESetErrorState SetErrorState)
 {
+	FPyConversionResult Result = FPyConversionResult::Failure();
+
 	PyTypeObject* PyStructType = FPyWrapperTypeRegistry::Get().GetWrappedStructType(StructType);
-	FPyWrapperStructPtr PyStruct = FPyWrapperStructPtr::StealReference(FPyWrapperStruct::CastPyObject(PyObj, PyStructType));
+	FPyWrapperStructPtr PyStruct = FPyWrapperStructPtr::StealReference(FPyWrapperStruct::CastPyObject(PyObj, PyStructType, &Result));
 	if (PyStruct)
 	{
 		StructType->CopyScriptStruct(StructInstance, PyStruct->StructInstance);
-		return true;
 	}
 
-	PYCONVERSION_RETURN(false, TEXT("Nativize"), *FString::Printf(TEXT("Cannot nativize '%s' as '%s'"), *PyUtil::GetFriendlyTypename(PyObj), *StructType->GetName()));
+	PYCONVERSION_RETURN(Result, TEXT("Nativize"), *FString::Printf(TEXT("Cannot nativize '%s' as '%s'"), *PyUtil::GetFriendlyTypename(PyObj), *StructType->GetName()));
 }
 
-bool PythonizeStruct(UScriptStruct* StructType, const void* StructInstance, PyObject*& OutPyObj, const ESetErrorState SetErrorState)
+FPyConversionResult PythonizeStruct(UScriptStruct* StructType, const void* StructInstance, PyObject*& OutPyObj, const ESetErrorState SetErrorState)
 {
 	OutPyObj = (PyObject*)FPyWrapperStructFactory::Get().CreateInstance(StructType, (void*)StructInstance, FPyWrapperOwnerContext(), EPyConversionMethod::Copy);
-	return true;
+	return FPyConversionResult::Success();
 }
 
 template <typename T>
-bool NativizeSigned(PyObject* PyObj, T& OutVal, const ESetErrorState SetErrorState, const TCHAR* InErrorType)
+FPyConversionResult NativizeSigned(PyObject* PyObj, T& OutVal, const ESetErrorState SetErrorState, const TCHAR* InErrorType)
 {
+	// Booleans subclass integer, so exclude those explicitly
+	if (PyBool_Check(PyObj))
+	{
+		return FPyConversionResult::Failure();
+	}
+
 #if PY_MAJOR_VERSION < 3
 	if (PyInt_Check(PyObj))
 	{
 		OutVal = PyInt_AsLong(PyObj);
-		return true;
+		return FPyConversionResult::Success();
 	}
 #endif	// PY_MAJOR_VERSION < 3
 
 	if (PyLong_Check(PyObj))
 	{
 		OutVal = PyLong_AsLongLong(PyObj);
-		return true;
+		return FPyConversionResult::Success();
 	}
 
 	if (PyFloat_Check(PyObj))
 	{
 		OutVal = PyFloat_AsDouble(PyObj);
-		return true;
+		return FPyConversionResult::SuccessWithCoercion();
 	}
 
-	PYCONVERSION_RETURN(false, TEXT("Nativize"), *FString::Printf(TEXT("Cannot nativize '%s' as '%s'"), *PyUtil::GetFriendlyTypename(PyObj), InErrorType));
+	PYCONVERSION_RETURN(FPyConversionResult::Failure(), TEXT("Nativize"), *FString::Printf(TEXT("Cannot nativize '%s' as '%s'"), *PyUtil::GetFriendlyTypename(PyObj), InErrorType));
 }
 
 template <typename T>
-bool NativizeUnsigned(PyObject* PyObj, T& OutVal, const ESetErrorState SetErrorState, const TCHAR* InErrorType)
+FPyConversionResult NativizeUnsigned(PyObject* PyObj, T& OutVal, const ESetErrorState SetErrorState, const TCHAR* InErrorType)
 {
+	// Booleans subclass integer, so exclude those explicitly
+	if (PyBool_Check(PyObj))
+	{
+		return FPyConversionResult::Failure();
+	}
+
 #if PY_MAJOR_VERSION < 3
 	if (PyInt_Check(PyObj))
 	{
 		OutVal = PyInt_AsSsize_t(PyObj);
-		return true;
+		return FPyConversionResult::Success();
 	}
 #endif	// PY_MAJOR_VERSION < 3
 
 	if (PyLong_Check(PyObj))
 	{
 		OutVal = PyLong_AsUnsignedLongLong(PyObj);
-		return true;
+		return FPyConversionResult::Success();
 	}
 
 	if (PyFloat_Check(PyObj))
 	{
 		OutVal = PyFloat_AsDouble(PyObj);
-		return true;
+		return FPyConversionResult::SuccessWithCoercion();
 	}
 
-	PYCONVERSION_RETURN(false, TEXT("Nativize"), *FString::Printf(TEXT("Cannot nativize '%s' as '%s'"), *PyUtil::GetFriendlyTypename(PyObj), InErrorType));
+	PYCONVERSION_RETURN(FPyConversionResult::Failure(), TEXT("Nativize"), *FString::Printf(TEXT("Cannot nativize '%s' as '%s'"), *PyUtil::GetFriendlyTypename(PyObj), InErrorType));
 }
 
 template <typename T>
-bool PythonizeSigned(const T Val, PyObject*& OutPyObj, const ESetErrorState SetErrorState, const TCHAR* InErrorType)
+FPyConversionResult NativizeReal(PyObject* PyObj, T& OutVal, const ESetErrorState SetErrorState, const TCHAR* InErrorType)
+{
+	// Booleans subclass integer, so exclude those explicitly
+	if (PyBool_Check(PyObj))
+	{
+		return FPyConversionResult::Failure();
+	}
+
+#if PY_MAJOR_VERSION < 3
+	if (PyInt_Check(PyObj))
+	{
+		OutVal = PyInt_AsSsize_t(PyObj);
+		return FPyConversionResult::SuccessWithCoercion();
+	}
+#endif	// PY_MAJOR_VERSION < 3
+
+	if (PyLong_Check(PyObj))
+	{
+		OutVal = PyLong_AsUnsignedLongLong(PyObj);
+		return FPyConversionResult::SuccessWithCoercion();
+	}
+
+	if (PyFloat_Check(PyObj))
+	{
+		OutVal = PyFloat_AsDouble(PyObj);
+		return FPyConversionResult::Success();
+	}
+
+	PYCONVERSION_RETURN(FPyConversionResult::Failure(), TEXT("Nativize"), *FString::Printf(TEXT("Cannot nativize '%s' as '%s'"), *PyUtil::GetFriendlyTypename(PyObj), InErrorType));
+}
+
+template <typename T>
+FPyConversionResult PythonizeSigned(const T Val, PyObject*& OutPyObj, const ESetErrorState SetErrorState, const TCHAR* InErrorType)
 {
 #if PY_MAJOR_VERSION < 3
 	if (Val >= LONG_MIN && Val <= LONG_MAX)
@@ -124,11 +169,11 @@ bool PythonizeSigned(const T Val, PyObject*& OutPyObj, const ESetErrorState SetE
 		OutPyObj = PyLong_FromLongLong(Val);
 	}
 
-	return true;
+	return FPyConversionResult::Success();
 }
 
 template <typename T>
-bool PythonizeUnsigned(const T Val, PyObject*& OutPyObj, const ESetErrorState SetErrorState, const TCHAR* InErrorType)
+FPyConversionResult PythonizeUnsigned(const T Val, PyObject*& OutPyObj, const ESetErrorState SetErrorState, const TCHAR* InErrorType)
 {
 #if PY_MAJOR_VERSION < 3
 	if (Val <= LONG_MAX)
@@ -141,56 +186,56 @@ bool PythonizeUnsigned(const T Val, PyObject*& OutPyObj, const ESetErrorState Se
 		OutPyObj = PyLong_FromUnsignedLongLong(Val);
 	}
 
-	return true;
+	return FPyConversionResult::Success();
 }
 
 template <typename T>
-bool PythonizeReal(const T Val, PyObject*& OutPyObj, const ESetErrorState SetErrorState, const TCHAR* InErrorType)
+FPyConversionResult PythonizeReal(const T Val, PyObject*& OutPyObj, const ESetErrorState SetErrorState, const TCHAR* InErrorType)
 {
 	OutPyObj = PyFloat_FromDouble(Val);
-	return true;
+	return FPyConversionResult::Success();
 }
 
 } // namespace Internal
 
-bool Nativize(PyObject* PyObj, bool& OutVal, const ESetErrorState SetErrorState)
+FPyConversionResult Nativize(PyObject* PyObj, bool& OutVal, const ESetErrorState SetErrorState)
 {
 	if (PyObj == Py_True)
 	{
 		OutVal = true;
-		return true;
+		return FPyConversionResult::Success();
 	}
 
 	if (PyObj == Py_False)
 	{
 		OutVal = false;
-		return true;
+		return FPyConversionResult::Success();
 	}
 
 	if (PyObj == Py_None)
 	{
 		OutVal = false;
-		return true;
+		return FPyConversionResult::Success();
 	}
 	
 #if PY_MAJOR_VERSION < 3
 	if (PyInt_Check(PyObj))
 	{
 		OutVal = PyInt_AsLong(PyObj) != 0;
-		return true;
+		return FPyConversionResult::SuccessWithCoercion();
 	}
 #endif	// PY_MAJOR_VERSION < 3
 
 	if (PyLong_Check(PyObj))
 	{
 		OutVal = PyLong_AsLongLong(PyObj) != 0;
-		return true;
+		return FPyConversionResult::SuccessWithCoercion();
 	}
 
-	PYCONVERSION_RETURN(false, TEXT("Nativize"), *FString::Printf(TEXT("Cannot nativize '%s' as 'bool'"), *PyUtil::GetFriendlyTypename(PyObj)));
+	PYCONVERSION_RETURN(FPyConversionResult::Failure(), TEXT("Nativize"), *FString::Printf(TEXT("Cannot nativize '%s' as 'bool'"), *PyUtil::GetFriendlyTypename(PyObj)));
 }
 
-bool Pythonize(const bool Val, PyObject*& OutPyObj, const ESetErrorState SetErrorState)
+FPyConversionResult Pythonize(const bool Val, PyObject*& OutPyObj, const ESetErrorState SetErrorState)
 {
 	if (Val)
 	{
@@ -203,110 +248,110 @@ bool Pythonize(const bool Val, PyObject*& OutPyObj, const ESetErrorState SetErro
 		OutPyObj = Py_False;
 	}
 
-	return true;
+	return FPyConversionResult::Success();
 }
 
-bool Nativize(PyObject* PyObj, int8& OutVal, const ESetErrorState SetErrorState)
+FPyConversionResult Nativize(PyObject* PyObj, int8& OutVal, const ESetErrorState SetErrorState)
 {
 	return Internal::NativizeSigned(PyObj, OutVal, SetErrorState, TEXT("int8"));
 }
 
-bool Pythonize(const int8 Val, PyObject*& OutPyObj, const ESetErrorState SetErrorState)
+FPyConversionResult Pythonize(const int8 Val, PyObject*& OutPyObj, const ESetErrorState SetErrorState)
 {
 	return Internal::PythonizeSigned(Val, OutPyObj, SetErrorState, TEXT("int8"));
 }
 
-bool Nativize(PyObject* PyObj, uint8& OutVal, const ESetErrorState SetErrorState)
+FPyConversionResult Nativize(PyObject* PyObj, uint8& OutVal, const ESetErrorState SetErrorState)
 {
 	return Internal::NativizeUnsigned(PyObj, OutVal, SetErrorState, TEXT("uint8"));
 }
 
-bool Pythonize(const uint8 Val, PyObject*& OutPyObj, const ESetErrorState SetErrorState)
+FPyConversionResult Pythonize(const uint8 Val, PyObject*& OutPyObj, const ESetErrorState SetErrorState)
 {
 	return Internal::PythonizeUnsigned(Val, OutPyObj, SetErrorState, TEXT("uint8"));
 }
 
-bool Nativize(PyObject* PyObj, int16& OutVal, const ESetErrorState SetErrorState)
+FPyConversionResult Nativize(PyObject* PyObj, int16& OutVal, const ESetErrorState SetErrorState)
 {
 	return Internal::NativizeSigned(PyObj, OutVal, SetErrorState, TEXT("int16"));
 }
 
-bool Pythonize(const int16 Val, PyObject*& OutPyObj, const ESetErrorState SetErrorState)
+FPyConversionResult Pythonize(const int16 Val, PyObject*& OutPyObj, const ESetErrorState SetErrorState)
 {
 	return Internal::PythonizeSigned(Val, OutPyObj, SetErrorState, TEXT("int16"));
 }
 
-bool Nativize(PyObject* PyObj, uint16& OutVal, const ESetErrorState SetErrorState)
+FPyConversionResult Nativize(PyObject* PyObj, uint16& OutVal, const ESetErrorState SetErrorState)
 {
 	return Internal::NativizeUnsigned(PyObj, OutVal, SetErrorState, TEXT("uint16"));
 }
 
-bool Pythonize(const uint16 Val, PyObject*& OutPyObj, const ESetErrorState SetErrorState)
+FPyConversionResult Pythonize(const uint16 Val, PyObject*& OutPyObj, const ESetErrorState SetErrorState)
 {
 	return Internal::PythonizeUnsigned(Val, OutPyObj, SetErrorState, TEXT("uint16"));
 }
 
-bool Nativize(PyObject* PyObj, int32& OutVal, const ESetErrorState SetErrorState)
+FPyConversionResult Nativize(PyObject* PyObj, int32& OutVal, const ESetErrorState SetErrorState)
 {
 	return Internal::NativizeSigned(PyObj, OutVal, SetErrorState, TEXT("int32"));
 }
 
-bool Pythonize(const int32 Val, PyObject*& OutPyObj, const ESetErrorState SetErrorState)
+FPyConversionResult Pythonize(const int32 Val, PyObject*& OutPyObj, const ESetErrorState SetErrorState)
 {
 	return Internal::PythonizeSigned(Val, OutPyObj, SetErrorState, TEXT("int32"));
 }
 
-bool Nativize(PyObject* PyObj, uint32& OutVal, const ESetErrorState SetErrorState)
+FPyConversionResult Nativize(PyObject* PyObj, uint32& OutVal, const ESetErrorState SetErrorState)
 {
 	return Internal::NativizeUnsigned(PyObj, OutVal, SetErrorState, TEXT("uint32"));
 }
 
-bool Pythonize(const uint32 Val, PyObject*& OutPyObj, const ESetErrorState SetErrorState)
+FPyConversionResult Pythonize(const uint32 Val, PyObject*& OutPyObj, const ESetErrorState SetErrorState)
 {
 	return Internal::PythonizeUnsigned(Val, OutPyObj, SetErrorState, TEXT("uint32"));
 }
 
-bool Nativize(PyObject* PyObj, int64& OutVal, const ESetErrorState SetErrorState)
+FPyConversionResult Nativize(PyObject* PyObj, int64& OutVal, const ESetErrorState SetErrorState)
 {
 	return Internal::NativizeSigned(PyObj, OutVal, SetErrorState, TEXT("int64"));
 }
 
-bool Pythonize(const int64 Val, PyObject*& OutPyObj, const ESetErrorState SetErrorState)
+FPyConversionResult Pythonize(const int64 Val, PyObject*& OutPyObj, const ESetErrorState SetErrorState)
 {
 	return Internal::PythonizeSigned(Val, OutPyObj, SetErrorState, TEXT("int64"));
 }
 
-bool Nativize(PyObject* PyObj, uint64& OutVal, const ESetErrorState SetErrorState)
+FPyConversionResult Nativize(PyObject* PyObj, uint64& OutVal, const ESetErrorState SetErrorState)
 {
 	return Internal::NativizeUnsigned(PyObj, OutVal, SetErrorState, TEXT("uint64"));
 }
 
-bool Pythonize(const uint64 Val, PyObject*& OutPyObj, const ESetErrorState SetErrorState)
+FPyConversionResult Pythonize(const uint64 Val, PyObject*& OutPyObj, const ESetErrorState SetErrorState)
 {
 	return Internal::PythonizeUnsigned(Val, OutPyObj, SetErrorState, TEXT("uint64"));
 }
 
-bool Nativize(PyObject* PyObj, float& OutVal, const ESetErrorState SetErrorState)
+FPyConversionResult Nativize(PyObject* PyObj, float& OutVal, const ESetErrorState SetErrorState)
 {
-	return Internal::NativizeSigned(PyObj, OutVal, SetErrorState, TEXT("float"));
+	return Internal::NativizeReal(PyObj, OutVal, SetErrorState, TEXT("float"));
 }
 
-bool Pythonize(const float Val, PyObject*& OutPyObj, const ESetErrorState SetErrorState)
+FPyConversionResult Pythonize(const float Val, PyObject*& OutPyObj, const ESetErrorState SetErrorState)
 {
 	return Internal::PythonizeReal(Val, OutPyObj, SetErrorState, TEXT("float"));
 }
 
-bool Nativize(PyObject* PyObj, double& OutVal, const ESetErrorState SetErrorState)
+FPyConversionResult Nativize(PyObject* PyObj, double& OutVal, const ESetErrorState SetErrorState)
 {
-	return Internal::NativizeSigned(PyObj, OutVal, SetErrorState, TEXT("double"));
+	return Internal::NativizeReal(PyObj, OutVal, SetErrorState, TEXT("double"));
 }
 
-bool Pythonize(const double Val, PyObject*& OutPyObj, const ESetErrorState SetErrorState)
+FPyConversionResult Pythonize(const double Val, PyObject*& OutPyObj, const ESetErrorState SetErrorState)
 {
 	return Internal::PythonizeReal(Val, OutPyObj, SetErrorState, TEXT("double"));
 }
 
-bool Nativize(PyObject* PyObj, FString& OutVal, const ESetErrorState SetErrorState)
+FPyConversionResult Nativize(PyObject* PyObj, FString& OutVal, const ESetErrorState SetErrorState)
 {
 	if (PyUnicode_Check(PyObj))
 	{
@@ -315,7 +360,7 @@ bool Nativize(PyObject* PyObj, FString& OutVal, const ESetErrorState SetErrorSta
 		{
 			const char* PyUtf8Buffer = PyBytes_AsString(PyBytesObj);
 			OutVal = UTF8_TO_TCHAR(PyUtf8Buffer);
-			return true;
+			return FPyConversionResult::Success();
 		}
 	}
 
@@ -324,21 +369,21 @@ bool Nativize(PyObject* PyObj, FString& OutVal, const ESetErrorState SetErrorSta
 	{
 		const char* PyUtf8Buffer = PyString_AsString(PyObj);
 		OutVal = UTF8_TO_TCHAR(PyUtf8Buffer);
-		return true;
+		return FPyConversionResult::Success();
 	}
 #endif	// PY_MAJOR_VERSION < 3
 
 	if (PyObject_IsInstance(PyObj, (PyObject*)&PyWrapperNameType) == 1)
 	{
 		FPyWrapperName* PyWrappedName = (FPyWrapperName*)PyObj;
-		OutVal = PyWrappedName->Name.ToString();
-		return true;
+		OutVal = PyWrappedName->Value.ToString();
+		return FPyConversionResult::SuccessWithCoercion();
 	}
 
-	PYCONVERSION_RETURN(false, TEXT("Nativize"), *FString::Printf(TEXT("Cannot nativize '%s' as 'String'"), *PyUtil::GetFriendlyTypename(PyObj)));
+	PYCONVERSION_RETURN(FPyConversionResult::Failure(), TEXT("Nativize"), *FString::Printf(TEXT("Cannot nativize '%s' as 'String'"), *PyUtil::GetFriendlyTypename(PyObj)));
 }
 
-bool Pythonize(const FString& Val, PyObject*& OutPyObj, const ESetErrorState SetErrorState)
+FPyConversionResult Pythonize(const FString& Val, PyObject*& OutPyObj, const ESetErrorState SetErrorState)
 {
 #if PY_MAJOR_VERSION < 3
 	if (FCString::IsPureAnsi(*Val))
@@ -351,70 +396,128 @@ bool Pythonize(const FString& Val, PyObject*& OutPyObj, const ESetErrorState Set
 		OutPyObj = PyUnicode_FromString(TCHAR_TO_UTF8(*Val));
 	}
 
-	return true;
+	return FPyConversionResult::Success();
 }
 
-bool Nativize(PyObject* PyObj, FName& OutVal, const ESetErrorState SetErrorState)
+FPyConversionResult Nativize(PyObject* PyObj, FName& OutVal, const ESetErrorState SetErrorState)
 {
 	if (PyObject_IsInstance(PyObj, (PyObject*)&PyWrapperNameType) == 1)
 	{
 		FPyWrapperName* PyWrappedName = (FPyWrapperName*)PyObj;
-		OutVal = PyWrappedName->Name;
-		return true;
+		OutVal = PyWrappedName->Value;
+		return FPyConversionResult::Success();
 	}
 
 	FString NameStr;
 	if (Nativize(PyObj, NameStr, ESetErrorState::No))
 	{
 		OutVal = *NameStr;
-		return true;
+		return FPyConversionResult::SuccessWithCoercion();
 	}
 
-	PYCONVERSION_RETURN(false, TEXT("Nativize"), *FString::Printf(TEXT("Cannot nativize '%s' as 'Name'"), *PyUtil::GetFriendlyTypename(PyObj)));
+	PYCONVERSION_RETURN(FPyConversionResult::Failure(), TEXT("Nativize"), *FString::Printf(TEXT("Cannot nativize '%s' as 'Name'"), *PyUtil::GetFriendlyTypename(PyObj)));
 }
 
-bool Pythonize(const FName& Val, PyObject*& OutPyObj, const ESetErrorState SetErrorState)
+FPyConversionResult Pythonize(const FName& Val, PyObject*& OutPyObj, const ESetErrorState SetErrorState)
 {
 	OutPyObj = (PyObject*)FPyWrapperNameFactory::Get().CreateInstance(Val);
-	return true;
+	return FPyConversionResult::Success();
 }
 
-bool Nativize(PyObject* PyObj, FText& OutVal, const ESetErrorState SetErrorState)
+FPyConversionResult Nativize(PyObject* PyObj, FText& OutVal, const ESetErrorState SetErrorState)
 {
 	if (PyObject_IsInstance(PyObj, (PyObject*)&PyWrapperTextType) == 1)
 	{
 		FPyWrapperText* PyWrappedName = (FPyWrapperText*)PyObj;
-		OutVal = PyWrappedName->Text;
-		return true;
+		OutVal = PyWrappedName->Value;
+		return FPyConversionResult::Success();
 	}
 
 	FString TextStr;
 	if (Nativize(PyObj, TextStr, ESetErrorState::No))
 	{
 		OutVal = FText::AsCultureInvariant(TextStr);
-		return true;
+		return FPyConversionResult::SuccessWithCoercion();
 	}
 
-	PYCONVERSION_RETURN(false, TEXT("Nativize"), *FString::Printf(TEXT("Cannot nativize '%s' as 'Text'"), *PyUtil::GetFriendlyTypename(PyObj)));
+	PYCONVERSION_RETURN(FPyConversionResult::Failure(), TEXT("Nativize"), *FString::Printf(TEXT("Cannot nativize '%s' as 'Text'"), *PyUtil::GetFriendlyTypename(PyObj)));
 }
 
-bool Pythonize(const FText& Val, PyObject*& OutPyObj, const ESetErrorState SetErrorState)
+FPyConversionResult Pythonize(const FText& Val, PyObject*& OutPyObj, const ESetErrorState SetErrorState)
 {
 	OutPyObj = (PyObject*)FPyWrapperTextFactory::Get().CreateInstance(Val);
-	return true;
+	return FPyConversionResult::Success();
 }
 
-bool Nativize(PyObject* PyObj, UObject*& OutVal, const ESetErrorState SetErrorState)
+FPyConversionResult Nativize(PyObject* PyObj, void*& OutVal, const ESetErrorState SetErrorState)
+{
+	// CObject was removed in Python 3.2
+#if !(PY_MAJOR_VERSION >= 3 && PY_MINOR_VERSION >= 2)
+	if (PyCObject_Check(PyObj))
+	{
+		OutVal = PyCObject_AsVoidPtr(PyObj);
+		return FPyConversionResult::Success();
+	}
+#endif	// !(PY_MAJOR_VERSION >= 3 && PY_MINOR_VERSION >= 2)
+
+	// Capsule was added in Python 2.7 and 3.1
+#if (PY_MAJOR_VERSION >= 2 && PY_MINOR_VERSION >= 7) || (PY_MAJOR_VERSION >= 3 && PY_MINOR_VERSION >= 1)
+	if (PyCapsule_CheckExact(PyObj))
+	{
+		OutVal = PyCapsule_GetPointer(PyObj, PyCapsule_GetName(PyObj));
+		return FPyConversionResult::Success();
+	}
+#endif	// (PY_MAJOR_VERSION >= 2 && PY_MINOR_VERSION >= 7) || (PY_MAJOR_VERSION >= 3 && PY_MINOR_VERSION >= 1)
+
+	if (PyObj == Py_None)
+	{
+		OutVal = nullptr;
+		return FPyConversionResult::Success();
+	}
+
+	{
+		uint64 PtrValue = 0;
+		if (PyConversion::Nativize(PyObj, PtrValue, ESetErrorState::No))
+		{
+			OutVal = (void*)PtrValue;
+			return FPyConversionResult::SuccessWithCoercion();
+		}
+	}
+
+	PYCONVERSION_RETURN(FPyConversionResult::Failure(), TEXT("Nativize"), *FString::Printf(TEXT("Cannot nativize '%s' as 'void*'"), *PyUtil::GetFriendlyTypename(PyObj)));
+}
+
+FPyConversionResult Pythonize(void* Val, PyObject*& OutPyObj, const ESetErrorState SetErrorState)
+{
+	if (Val)
+	{
+		// Use Capsule for Python 3.1+, and CObject for older versions
+#if PY_MAJOR_VERSION >= 3 && PY_MINOR_VERSION >= 1
+		OutPyObj = PyCapsule_New(Val, nullptr, nullptr);
+#else	// PY_MAJOR_VERSION >= 3 && PY_MINOR_VERSION >= 1
+		OutPyObj = PyCObject_FromVoidPtr(Val, nullptr);
+#endif	// PY_MAJOR_VERSION >= 3 && PY_MINOR_VERSION >= 1
+	}
+	else
+	{
+		Py_INCREF(Py_None);
+		OutPyObj = Py_None;
+	}
+
+	return FPyConversionResult::Success();
+}
+
+FPyConversionResult Nativize(PyObject* PyObj, UObject*& OutVal, const ESetErrorState SetErrorState)
 {
 	return NativizeObject(PyObj, OutVal, UObject::StaticClass(), SetErrorState);
 }
 
-bool Pythonize(UObject* Val, PyObject*& OutPyObj, const ESetErrorState SetErrorState)
+FPyConversionResult Pythonize(UObject* Val, PyObject*& OutPyObj, const ESetErrorState SetErrorState)
 {
 	return PythonizeObject(Val, OutPyObj, SetErrorState);
 }
 
-bool NativizeObject(PyObject* PyObj, UObject*& OutVal, UClass* ExpectedType, const ESetErrorState SetErrorState)
+FPyConversionResult NativizeObject(PyObject* PyObj, UObject*& OutVal, UClass* ExpectedType, const ESetErrorState SetErrorState)
 {
 	if (PyObject_IsInstance(PyObj, (PyObject*)&PyWrapperObjectType) == 1)
 	{
@@ -422,20 +525,20 @@ bool NativizeObject(PyObject* PyObj, UObject*& OutVal, UClass* ExpectedType, con
 		if (!ExpectedType || PyWrappedObj->ObjectInstance->IsA(ExpectedType))
 		{
 			OutVal = PyWrappedObj->ObjectInstance;
-			return true;
+			return FPyConversionResult::Success();
 		}
 	}
 
 	if (PyObj == Py_None)
 	{
 		OutVal = nullptr;
-		return true;
+		return FPyConversionResult::Success();
 	}
 
-	PYCONVERSION_RETURN(false, TEXT("NativizeObject"), *FString::Printf(TEXT("Cannot nativize '%s' as 'Object*' (allowed Class type: '%s')"), *PyUtil::GetFriendlyTypename(PyObj), ExpectedType ? *ExpectedType->GetName() : TEXT("<any>")));
+	PYCONVERSION_RETURN(FPyConversionResult::Failure(), TEXT("NativizeObject"), *FString::Printf(TEXT("Cannot nativize '%s' as 'Object*' (allowed Class type: '%s')"), *PyUtil::GetFriendlyTypename(PyObj), ExpectedType ? *ExpectedType->GetName() : TEXT("<any>")));
 }
 
-bool PythonizeObject(UObject* Val, PyObject*& OutPyObj, const ESetErrorState SetErrorState)
+FPyConversionResult PythonizeObject(UObject* Val, PyObject*& OutPyObj, const ESetErrorState SetErrorState)
 {
 	if (Val)
 	{
@@ -447,7 +550,7 @@ bool PythonizeObject(UObject* Val, PyObject*& OutPyObj, const ESetErrorState Set
 		OutPyObj = Py_None;
 	}
 
-	return true;
+	return FPyConversionResult::Success();
 }
 
 PyObject* PythonizeObject(UObject* Val, const ESetErrorState SetErrorState)
@@ -457,7 +560,7 @@ PyObject* PythonizeObject(UObject* Val, const ESetErrorState SetErrorState)
 	return Obj;
 }
 
-bool NativizeClass(PyObject* PyObj, UClass*& OutVal, UClass* ExpectedType, const ESetErrorState SetErrorState)
+FPyConversionResult NativizeClass(PyObject* PyObj, UClass*& OutVal, UClass* ExpectedType, const ESetErrorState SetErrorState)
 {
 	UClass* Class = nullptr;
 
@@ -471,14 +574,14 @@ bool NativizeClass(PyObject* PyObj, UClass*& OutVal, UClass* ExpectedType, const
 		if (!Class || !ExpectedType || Class->IsChildOf(ExpectedType))
 		{
 			OutVal = Class;
-			return true;
+			return FPyConversionResult::Success();
 		}
 	}
 
-	PYCONVERSION_RETURN(false, TEXT("NativizeClass"), *FString::Printf(TEXT("Cannot nativize '%s' as 'Class*' (allowed Class type: '%s')"), *PyUtil::GetFriendlyTypename(PyObj), ExpectedType ? *ExpectedType->GetName() : TEXT("<any>")));
+	PYCONVERSION_RETURN(FPyConversionResult::Failure(), TEXT("NativizeClass"), *FString::Printf(TEXT("Cannot nativize '%s' as 'Class*' (allowed Class type: '%s')"), *PyUtil::GetFriendlyTypename(PyObj), ExpectedType ? *ExpectedType->GetName() : TEXT("<any>")));
 }
 
-bool PythonizeClass(UClass* Val, PyObject*& OutPyObj, const ESetErrorState SetErrorState)
+FPyConversionResult PythonizeClass(UClass* Val, PyObject*& OutPyObj, const ESetErrorState SetErrorState)
 {
 	return PythonizeObject(Val, OutPyObj, SetErrorState);
 }
@@ -490,7 +593,7 @@ PyObject* PythonizeClass(UClass* Val, const ESetErrorState SetErrorState)
 	return Obj;
 }
 
-bool NativizeProperty(PyObject* PyObj, const UProperty* Prop, void* ValueAddr, const FPyWrapperOwnerContext& InChangeOwner, const ESetErrorState SetErrorState)
+FPyConversionResult NativizeProperty(PyObject* PyObj, const UProperty* Prop, void* ValueAddr, const FPyWrapperOwnerContext& InChangeOwner, const ESetErrorState SetErrorState)
 {
 #define PYCONVERSION_PROPERTY_RETURN(RESULT) \
 	PYCONVERSION_RETURN(RESULT, TEXT("NativizeProperty"), *FString::Printf(TEXT("Cannot nativize '%s' as '%s' (%s)"), *PyUtil::GetFriendlyTypename(PyObj), *Prop->GetName(), *Prop->GetClass()->GetName()))
@@ -505,10 +608,10 @@ bool NativizeProperty(PyObject* PyObj, const UProperty* Prop, void* ValueAddr, c
 			{
 				Prop->CopySingleValue(static_cast<uint8*>(ValueAddr) + (Prop->ElementSize * ArrIndex), FPyWrapperFixedArray::GetItemPtr(PyFixedArray, ArrIndex));
 			}
-			return true;
+			return FPyConversionResult::Success();
 		}
 
-		PYCONVERSION_PROPERTY_RETURN(false);
+		PYCONVERSION_PROPERTY_RETURN(FPyConversionResult::Failure());
 	}
 
 	return NativizeProperty_Direct(PyObj, Prop, ValueAddr, InChangeOwner, SetErrorState);
@@ -516,18 +619,18 @@ bool NativizeProperty(PyObject* PyObj, const UProperty* Prop, void* ValueAddr, c
 #undef PYCONVERSION_PROPERTY_RETURN
 }
 
-bool PythonizeProperty(const UProperty* Prop, const void* ValueAddr, PyObject*& OutPyObj, const EPyConversionMethod ConversionMethod, PyObject* OwnerPyObj, const ESetErrorState SetErrorState)
+FPyConversionResult PythonizeProperty(const UProperty* Prop, const void* ValueAddr, PyObject*& OutPyObj, const EPyConversionMethod ConversionMethod, PyObject* OwnerPyObj, const ESetErrorState SetErrorState)
 {
 	if (Prop->ArrayDim > 1)
 	{
-		OutPyObj = (PyObject*)FPyWrapperFixedArrayFactory::Get().CreateInstance((void*)ValueAddr, Prop, FPyWrapperOwnerContext(OwnerPyObj, OwnerPyObj ? Prop->GetFName() : NAME_None), ConversionMethod);
-		return true;
+		OutPyObj = (PyObject*)FPyWrapperFixedArrayFactory::Get().CreateInstance((void*)ValueAddr, Prop, FPyWrapperOwnerContext(OwnerPyObj, OwnerPyObj ? Prop : nullptr), ConversionMethod);
+		return FPyConversionResult::Success();
 	}
 
 	return PythonizeProperty_Direct(Prop, ValueAddr, OutPyObj, ConversionMethod, OwnerPyObj, SetErrorState);
 }
 
-bool NativizeProperty_Direct(PyObject* PyObj, const UProperty* Prop, void* ValueAddr, const FPyWrapperOwnerContext& InChangeOwner, const ESetErrorState SetErrorState)
+FPyConversionResult NativizeProperty_Direct(PyObject* PyObj, const UProperty* Prop, void* ValueAddr, const FPyWrapperOwnerContext& InChangeOwner, const ESetErrorState SetErrorState)
 {
 #define PYCONVERSION_PROPERTY_RETURN(RESULT) \
 	PYCONVERSION_RETURN(RESULT, TEXT("NativizeProperty"), *FString::Printf(TEXT("Cannot nativize '%s' as '%s' (%s)"), *PyUtil::GetFriendlyTypename(PyObj), *Prop->GetName(), *Prop->GetClass()->GetName()))
@@ -536,8 +639,8 @@ bool NativizeProperty_Direct(PyObject* PyObj, const UProperty* Prop, void* Value
 	if (auto* CastProp = Cast<PROPTYPE>(Prop))										\
 	{																				\
 		PROPTYPE::TCppType NewValue;												\
-		const bool bResult = Nativize(PyObj, NewValue, SetErrorState);				\
-		if (bResult)																\
+		const FPyConversionResult Result = Nativize(PyObj, NewValue, SetErrorState);\
+		if (Result)																	\
 		{																			\
 			auto OldValue = CastProp->GetPropertyValue(ValueAddr);					\
 			if (OldValue != NewValue)												\
@@ -548,15 +651,15 @@ bool NativizeProperty_Direct(PyObject* PyObj, const UProperty* Prop, void* Value
 				});																	\
 			}																		\
 		}																			\
-		PYCONVERSION_PROPERTY_RETURN(bResult);										\
+		PYCONVERSION_PROPERTY_RETURN(Result);										\
 	}
 
 #define NATIVIZE_INLINE_PROPERTY(PROPTYPE)											\
 	if (auto* CastProp = Cast<PROPTYPE>(Prop))										\
 	{																				\
 		PROPTYPE::TCppType NewValue;												\
-		const bool bResult = Nativize(PyObj, NewValue, SetErrorState);				\
-		if (bResult)																\
+		const FPyConversionResult Result = Nativize(PyObj, NewValue, SetErrorState);\
+		if (Result)																	\
 		{																			\
 			auto* ValuePtr = static_cast<PROPTYPE::TCppType*>(ValueAddr);			\
 			if (!CastProp->Identical(ValuePtr, &NewValue, PPF_None))				\
@@ -567,7 +670,7 @@ bool NativizeProperty_Direct(PyObject* PyObj, const UProperty* Prop, void* Value
 				});																	\
 			}																		\
 		}																			\
-		PYCONVERSION_PROPERTY_RETURN(bResult);										\
+		PYCONVERSION_PROPERTY_RETURN(Result);										\
 	}
 
 	NATIVIZE_SETTER_PROPERTY(UBoolProperty);
@@ -588,14 +691,14 @@ bool NativizeProperty_Direct(PyObject* PyObj, const UProperty* Prop, void* Value
 	if (auto* CastProp = Cast<UEnumProperty>(Prop))
 	{
 		UNumericProperty* EnumInternalProp = CastProp->GetUnderlyingProperty();
-		PYCONVERSION_PROPERTY_RETURN(EnumInternalProp && NativizeProperty_Direct(PyObj, EnumInternalProp, ValueAddr, InChangeOwner, SetErrorState));
+		PYCONVERSION_PROPERTY_RETURN(EnumInternalProp ? NativizeProperty_Direct(PyObj, EnumInternalProp, ValueAddr, InChangeOwner, SetErrorState) : FPyConversionResult::Failure());
 	}
 
 	if (auto* CastProp = Cast<UClassProperty>(Prop))
 	{
 		UClass* NewValue = nullptr;
-		const bool bResult = NativizeClass(PyObj, NewValue, CastProp->MetaClass, SetErrorState);
-		if (bResult)
+		const FPyConversionResult Result = NativizeClass(PyObj, NewValue, CastProp->MetaClass, SetErrorState);
+		if (Result)
 		{
 			UObject* OldValue = CastProp->GetObjectPropertyValue(ValueAddr);
 			if (OldValue != NewValue)
@@ -606,14 +709,14 @@ bool NativizeProperty_Direct(PyObject* PyObj, const UProperty* Prop, void* Value
 				});
 			}
 		}
-		PYCONVERSION_PROPERTY_RETURN(bResult);
+		PYCONVERSION_PROPERTY_RETURN(Result);
 	}
 
 	if (auto* CastProp = Cast<UObjectPropertyBase>(Prop))
 	{
 		UObject* NewValue = nullptr;
-		const bool bResult = NativizeObject(PyObj, NewValue, CastProp->PropertyClass, SetErrorState);
-		if (bResult)
+		const FPyConversionResult Result = NativizeObject(PyObj, NewValue, CastProp->PropertyClass, SetErrorState);
+		if (Result)
 		{
 			UObject* OldValue = CastProp->GetObjectPropertyValue(ValueAddr);
 			if (OldValue != NewValue)
@@ -624,14 +727,14 @@ bool NativizeProperty_Direct(PyObject* PyObj, const UProperty* Prop, void* Value
 				});
 			}
 		}
-		PYCONVERSION_PROPERTY_RETURN(bResult);
+		PYCONVERSION_PROPERTY_RETURN(Result);
 	}
 
 	if (auto* CastProp = Cast<UInterfaceProperty>(Prop))
 	{
 		UObject* NewValue = nullptr;
-		const bool bResult = NativizeObject(PyObj, NewValue, CastProp->InterfaceClass, SetErrorState);
-		if (bResult)
+		const FPyConversionResult Result = NativizeObject(PyObj, NewValue, CastProp->InterfaceClass, SetErrorState);
+		if (Result)
 		{
 			UObject* OldValue = CastProp->GetPropertyValue(ValueAddr).GetObject();
 			if (OldValue != NewValue)
@@ -642,13 +745,14 @@ bool NativizeProperty_Direct(PyObject* PyObj, const UProperty* Prop, void* Value
 				});
 			}
 		}
-		PYCONVERSION_PROPERTY_RETURN(bResult);
+		PYCONVERSION_PROPERTY_RETURN(Result);
 	}
 
 	if (auto* CastProp = Cast<UStructProperty>(Prop))
 	{
+		FPyConversionResult Result = FPyConversionResult::Failure();
 		PyTypeObject* PyStructType = FPyWrapperTypeRegistry::Get().GetWrappedStructType(CastProp->Struct);
-		FPyWrapperStructPtr PyStruct = FPyWrapperStructPtr::StealReference(FPyWrapperStruct::CastPyObject(PyObj, PyStructType));
+		FPyWrapperStructPtr PyStruct = FPyWrapperStructPtr::StealReference(FPyWrapperStruct::CastPyObject(PyObj, PyStructType, &Result));
 		if (PyStruct)
 		{
 			if (!CastProp->Identical(ValueAddr, PyStruct->StructInstance, PPF_None))
@@ -658,15 +762,15 @@ bool NativizeProperty_Direct(PyObject* PyObj, const UProperty* Prop, void* Value
 					CastProp->Struct->CopyScriptStruct(ValueAddr, PyStruct->StructInstance);
 				});
 			}
-			return true;
 		}
-		PYCONVERSION_PROPERTY_RETURN(false);
+		PYCONVERSION_PROPERTY_RETURN(Result);
 	}
 
 	if (auto* CastProp = Cast<UDelegateProperty>(Prop))
 	{
+		FPyConversionResult Result = FPyConversionResult::Failure();
 		PyTypeObject* PyDelegateType = FPyWrapperTypeRegistry::Get().GetWrappedDelegateType(CastProp->SignatureFunction);
-		FPyWrapperDelegatePtr PyDelegate = FPyWrapperDelegatePtr::StealReference(FPyWrapperDelegate::CastPyObject(PyObj, PyDelegateType));
+		FPyWrapperDelegatePtr PyDelegate = FPyWrapperDelegatePtr::StealReference(FPyWrapperDelegate::CastPyObject(PyObj, PyDelegateType, &Result));
 		if (PyDelegate)
 		{
 			if (!CastProp->Identical(ValueAddr, PyDelegate->DelegateInstance, PPF_None))
@@ -676,15 +780,15 @@ bool NativizeProperty_Direct(PyObject* PyObj, const UProperty* Prop, void* Value
 					CastProp->SetPropertyValue(ValueAddr, *PyDelegate->DelegateInstance);
 				});
 			}
-			return true;
 		}
-		PYCONVERSION_PROPERTY_RETURN(false);
+		PYCONVERSION_PROPERTY_RETURN(Result);
 	}
 
 	if (auto* CastProp = Cast<UMulticastDelegateProperty>(Prop))
 	{
+		FPyConversionResult Result = FPyConversionResult::Failure();
 		PyTypeObject* PyDelegateType = FPyWrapperTypeRegistry::Get().GetWrappedDelegateType(CastProp->SignatureFunction);
-		FPyWrapperMulticastDelegatePtr PyDelegate = FPyWrapperMulticastDelegatePtr::StealReference(FPyWrapperMulticastDelegate::CastPyObject(PyObj, PyDelegateType));
+		FPyWrapperMulticastDelegatePtr PyDelegate = FPyWrapperMulticastDelegatePtr::StealReference(FPyWrapperMulticastDelegate::CastPyObject(PyObj, PyDelegateType, &Result));
 		if (PyDelegate)
 		{
 			if (!CastProp->Identical(ValueAddr, PyDelegate->DelegateInstance, PPF_None))
@@ -694,14 +798,14 @@ bool NativizeProperty_Direct(PyObject* PyObj, const UProperty* Prop, void* Value
 					CastProp->SetPropertyValue(ValueAddr, *PyDelegate->DelegateInstance);
 				});
 			}
-			return true;
 		}
-		PYCONVERSION_PROPERTY_RETURN(false);
+		PYCONVERSION_PROPERTY_RETURN(Result);
 	}
 
 	if (auto* CastProp = Cast<UArrayProperty>(Prop))
 	{
-		FPyWrapperArrayPtr PyArray = FPyWrapperArrayPtr::StealReference(FPyWrapperArray::CastPyObject(PyObj, &PyWrapperArrayType, CastProp->Inner));
+		FPyConversionResult Result = FPyConversionResult::Failure();
+		FPyWrapperArrayPtr PyArray = FPyWrapperArrayPtr::StealReference(FPyWrapperArray::CastPyObject(PyObj, &PyWrapperArrayType, CastProp->Inner, &Result));
 		if (PyArray)
 		{
 			if (!CastProp->Identical(ValueAddr, PyArray->ArrayInstance, PPF_None))
@@ -711,14 +815,14 @@ bool NativizeProperty_Direct(PyObject* PyObj, const UProperty* Prop, void* Value
 					CastProp->CopyCompleteValue(ValueAddr, PyArray->ArrayInstance);
 				});
 			}
-			return true;
 		}
-		PYCONVERSION_PROPERTY_RETURN(false);
+		PYCONVERSION_PROPERTY_RETURN(Result);
 	}
 
 	if (auto* CastProp = Cast<USetProperty>(Prop))
 	{
-		FPyWrapperSetPtr PySet = FPyWrapperSetPtr::StealReference(FPyWrapperSet::CastPyObject(PyObj, &PyWrapperSetType, CastProp->ElementProp));
+		FPyConversionResult Result = FPyConversionResult::Failure();
+		FPyWrapperSetPtr PySet = FPyWrapperSetPtr::StealReference(FPyWrapperSet::CastPyObject(PyObj, &PyWrapperSetType, CastProp->ElementProp, &Result));
 		if (PySet)
 		{
 			if (!CastProp->Identical(ValueAddr, PySet->SetInstance, PPF_None))
@@ -728,14 +832,14 @@ bool NativizeProperty_Direct(PyObject* PyObj, const UProperty* Prop, void* Value
 					CastProp->CopyCompleteValue(ValueAddr, PySet->SetInstance);
 				});
 			}
-			return true;
 		}
-		PYCONVERSION_PROPERTY_RETURN(false);
+		PYCONVERSION_PROPERTY_RETURN(Result);
 	}
 
 	if (auto* CastProp = Cast<UMapProperty>(Prop))
 	{
-		FPyWrapperMapPtr PyMap = FPyWrapperMapPtr::StealReference(FPyWrapperMap::CastPyObject(PyObj, &PyWrapperMapType, CastProp->KeyProp, CastProp->ValueProp));
+		FPyConversionResult Result = FPyConversionResult::Failure();
+		FPyWrapperMapPtr PyMap = FPyWrapperMapPtr::StealReference(FPyWrapperMap::CastPyObject(PyObj, &PyWrapperMapType, CastProp->KeyProp, CastProp->ValueProp, &Result));
 		if (PyMap)
 		{
 			if (!CastProp->Identical(ValueAddr, PyMap->MapInstance, PPF_None))
@@ -745,25 +849,24 @@ bool NativizeProperty_Direct(PyObject* PyObj, const UProperty* Prop, void* Value
 					CastProp->CopyCompleteValue(ValueAddr, PyMap->MapInstance);
 				});
 			}
-			return true;
 		}
-		PYCONVERSION_PROPERTY_RETURN(false);
+		PYCONVERSION_PROPERTY_RETURN(Result);
 	}
 
 #undef NATIVIZE_SETTER_PROPERTY
 #undef NATIVIZE_INLINE_PROPERTY
 
-	PYCONVERSION_RETURN(false, TEXT("NativizeProperty"), *FString::Printf(TEXT("Cannot nativize '%s' as '%s' (%s). %s conversion not implemented!"), *PyUtil::GetFriendlyTypename(PyObj), *Prop->GetName(), *Prop->GetClass()->GetName(), *Prop->GetClass()->GetName()));
+	PYCONVERSION_RETURN(FPyConversionResult::Failure(), TEXT("NativizeProperty"), *FString::Printf(TEXT("Cannot nativize '%s' as '%s' (%s). %s conversion not implemented!"), *PyUtil::GetFriendlyTypename(PyObj), *Prop->GetName(), *Prop->GetClass()->GetName(), *Prop->GetClass()->GetName()));
 
 #undef PYCONVERSION_PROPERTY_RETURN
 }
 
-bool PythonizeProperty_Direct(const UProperty* Prop, const void* ValueAddr, PyObject*& OutPyObj, const EPyConversionMethod ConversionMethod, PyObject* OwnerPyObj, const ESetErrorState SetErrorState)
+FPyConversionResult PythonizeProperty_Direct(const UProperty* Prop, const void* ValueAddr, PyObject*& OutPyObj, const EPyConversionMethod ConversionMethod, PyObject* OwnerPyObj, const ESetErrorState SetErrorState)
 {
 #define PYCONVERSION_PROPERTY_RETURN(RESULT) \
 	PYCONVERSION_RETURN(RESULT, TEXT("PythonizeProperty"), *FString::Printf(TEXT("Cannot pythonize '%s' (%s)"), *Prop->GetName(), *Prop->GetClass()->GetName()))
 
-	const FPyWrapperOwnerContext OwnerContext(OwnerPyObj, OwnerPyObj ? Prop->GetFName() : NAME_None);
+	const FPyWrapperOwnerContext OwnerContext(OwnerPyObj, OwnerPyObj ? Prop : nullptr);
 	OwnerContext.AssertValidConversionMethod(ConversionMethod);
 
 #define PYTHONIZE_GETTER_PROPERTY(PROPTYPE)											\
@@ -791,7 +894,7 @@ bool PythonizeProperty_Direct(const UProperty* Prop, const void* ValueAddr, PyOb
 	if (auto* CastProp = Cast<UEnumProperty>(Prop))
 	{
 		UNumericProperty* EnumInternalProp = CastProp->GetUnderlyingProperty();
-		PYCONVERSION_PROPERTY_RETURN(EnumInternalProp && PythonizeProperty_Direct(EnumInternalProp, ValueAddr, OutPyObj, ConversionMethod, OwnerPyObj, SetErrorState));
+		PYCONVERSION_PROPERTY_RETURN(EnumInternalProp ? PythonizeProperty_Direct(EnumInternalProp, ValueAddr, OutPyObj, ConversionMethod, OwnerPyObj, SetErrorState) : FPyConversionResult::Failure());
 	}
 
 	if (auto* CastProp = Cast<UClassProperty>(Prop))
@@ -818,61 +921,61 @@ bool PythonizeProperty_Direct(const UProperty* Prop, const void* ValueAddr, PyOb
 			Py_INCREF(Py_None);
 			OutPyObj = Py_None;
 		}
-		return true;
+		return FPyConversionResult::Success();
 	}
 
 	if (auto* CastProp = Cast<UStructProperty>(Prop))
 	{
 		OutPyObj = (PyObject*)FPyWrapperStructFactory::Get().CreateInstance(CastProp->Struct, (void*)ValueAddr, OwnerContext, ConversionMethod);
-		return true;
+		return FPyConversionResult::Success();
 	}
 
 	if (auto* CastProp = Cast<UDelegateProperty>(Prop))
 	{
 		const FScriptDelegate* Value = CastProp->GetPropertyValuePtr(ValueAddr);
 		OutPyObj = (PyObject*)FPyWrapperDelegateFactory::Get().CreateInstance(CastProp->SignatureFunction, (FScriptDelegate*)Value, OwnerContext, ConversionMethod);
-		return true;
+		return FPyConversionResult::Success();
 	}
 
 	if (auto* CastProp = Cast<UMulticastDelegateProperty>(Prop))
 	{
 		const FMulticastScriptDelegate* Value = CastProp->GetPropertyValuePtr(ValueAddr);
 		OutPyObj = (PyObject*)FPyWrapperMulticastDelegateFactory::Get().CreateInstance(CastProp->SignatureFunction, (FMulticastScriptDelegate*)Value, OwnerContext, ConversionMethod);
-		return true;
+		return FPyConversionResult::Success();
 	}
 
 	if (auto* CastProp = Cast<UArrayProperty>(Prop))
 	{
 		OutPyObj = (PyObject*)FPyWrapperArrayFactory::Get().CreateInstance((void*)ValueAddr, CastProp, OwnerContext, ConversionMethod);
-		return true;
+		return FPyConversionResult::Success();
 	}
 
 	if (auto* CastProp = Cast<USetProperty>(Prop))
 	{
 		OutPyObj = (PyObject*)FPyWrapperSetFactory::Get().CreateInstance((void*)ValueAddr, CastProp, OwnerContext, ConversionMethod);
-		return true;
+		return FPyConversionResult::Success();
 	}
 
 	if (auto* CastProp = Cast<UMapProperty>(Prop))
 	{
 		OutPyObj = (PyObject*)FPyWrapperMapFactory::Get().CreateInstance((void*)ValueAddr, CastProp, OwnerContext, ConversionMethod);
-		return true;
+		return FPyConversionResult::Success();
 	}
 
 #undef PYTHONIZE_GETTER_PROPERTY
 
-	PYCONVERSION_RETURN(false, TEXT("PythonizeProperty"), *FString::Printf(TEXT("Cannot pythonize '%s' (%s). %s conversion not implemented!"), *Prop->GetName(), *Prop->GetClass()->GetName(), *Prop->GetClass()->GetName()));
+	PYCONVERSION_RETURN(FPyConversionResult::Failure(), TEXT("PythonizeProperty"), *FString::Printf(TEXT("Cannot pythonize '%s' (%s). %s conversion not implemented!"), *Prop->GetName(), *Prop->GetClass()->GetName(), *Prop->GetClass()->GetName()));
 
 #undef PYCONVERSION_PROPERTY_RETURN
 }
 
-bool NativizeProperty_InContainer(PyObject* PyObj, const UProperty* Prop, void* BaseAddr, const int32 ArrayIndex, const FPyWrapperOwnerContext& InChangeOwner, const ESetErrorState SetErrorState)
+FPyConversionResult NativizeProperty_InContainer(PyObject* PyObj, const UProperty* Prop, void* BaseAddr, const int32 ArrayIndex, const FPyWrapperOwnerContext& InChangeOwner, const ESetErrorState SetErrorState)
 {
 	check(ArrayIndex < Prop->ArrayDim);
 	return NativizeProperty(PyObj, Prop, Prop->ContainerPtrToValuePtr<void>(BaseAddr, ArrayIndex), InChangeOwner, SetErrorState);
 }
 
-bool PythonizeProperty_InContainer(const UProperty* Prop, const void* BaseAddr, const int32 ArrayIndex, PyObject*& OutPyObj, const EPyConversionMethod ConversionMethod, PyObject* OwnerPyObj, const ESetErrorState SetErrorState)
+FPyConversionResult PythonizeProperty_InContainer(const UProperty* Prop, const void* BaseAddr, const int32 ArrayIndex, PyObject*& OutPyObj, const EPyConversionMethod ConversionMethod, PyObject* OwnerPyObj, const ESetErrorState SetErrorState)
 {
 	check(ArrayIndex < Prop->ArrayDim);
 	return PythonizeProperty(Prop, Prop->ContainerPtrToValuePtr<void>(BaseAddr, ArrayIndex), OutPyObj, ConversionMethod, OwnerPyObj, SetErrorState);
@@ -885,24 +988,16 @@ void EmitPropertyChangeNotifications(const FPyWrapperOwnerContext& InChangeOwner
 	{
 		auto AppendOwnerPropertyToChain = [&OutPropertyChain](const FPyWrapperOwnerContext& InOwnerContext) -> bool
 		{
-			const UStruct* LeafStruct = nullptr;
-			if (PyObject_IsInstance(InOwnerContext.GetOwnerObject(), (PyObject*)&PyWrapperObjectType) == 1)
-			{
-				LeafStruct = FPyWrapperObjectMetaData::GetClass((FPyWrapperObject*)InOwnerContext.GetOwnerObject());
-			}
-			else if (PyObject_IsInstance(InOwnerContext.GetOwnerObject(), (PyObject*)&PyWrapperStructType) == 1)
+			const UProperty* LeafProp = nullptr;
+			if (PyObject_IsInstance(InOwnerContext.GetOwnerObject(), (PyObject*)&PyWrapperObjectType) == 1 || PyObject_IsInstance(InOwnerContext.GetOwnerObject(), (PyObject*)&PyWrapperStructType) == 1)
 			{
-				LeafStruct = FPyWrapperStructMetaData::GetStruct((FPyWrapperStruct*)InOwnerContext.GetOwnerObject());
+				LeafProp = InOwnerContext.GetOwnerProperty();
 			}
 
-			if (LeafStruct)
+			if (LeafProp)
 			{
-				UProperty* LeafProp = LeafStruct->FindPropertyByName(InOwnerContext.GetOwnerPropertyName());
-				if (LeafProp)
-				{
-					OutPropertyChain.AddHead(LeafProp);
-					return true;
-				}
+				OutPropertyChain.AddHead((UProperty*)LeafProp);
+				return true;
 			}
 
 			return false;
diff --git a/Engine/Plugins/Experimental/PythonScriptPlugin/Source/PythonScriptPlugin/Private/PyConversion.h b/Engine/Plugins/Experimental/PythonScriptPlugin/Source/PythonScriptPlugin/Private/PyConversion.h
--- a/Engine/Plugins/Experimental/PythonScriptPlugin/Source/PythonScriptPlugin/Private/PyConversion.h
+++ b/Engine/Plugins/Experimental/PythonScriptPlugin/Source/PythonScriptPlugin/Private/PyConversion.h
@@ -4,6 +4,7 @@
 
 #include "IncludePython.h"
 #include "PyConversionMethod.h"
+#include "PyConversionResult.h"
 #include "PyWrapperOwnerContext.h"
 #include "UObject/Class.h"
 #include "Templates/EnableIf.h"
@@ -21,95 +22,99 @@ namespace PyConversion
 	enum class ESetErrorState : uint8 { No, Yes };
 
 	/** bool overload */
-	bool Nativize(PyObject* PyObj, bool& OutVal, const ESetErrorState SetErrorState = ESetErrorState::Yes);
-	bool Pythonize(const bool Val, PyObject*& OutPyObj, const ESetErrorState SetErrorState = ESetErrorState::Yes);
+	FPyConversionResult Nativize(PyObject* PyObj, bool& OutVal, const ESetErrorState SetErrorState = ESetErrorState::Yes);
+	FPyConversionResult Pythonize(const bool Val, PyObject*& OutPyObj, const ESetErrorState SetErrorState = ESetErrorState::Yes);
 
 	/** int8 overload */
-	bool Nativize(PyObject* PyObj, int8& OutVal, const ESetErrorState SetErrorState = ESetErrorState::Yes);
-	bool Pythonize(const int8 Val, PyObject*& OutPyObj, const ESetErrorState SetErrorState = ESetErrorState::Yes);
+	FPyConversionResult Nativize(PyObject* PyObj, int8& OutVal, const ESetErrorState SetErrorState = ESetErrorState::Yes);
+	FPyConversionResult Pythonize(const int8 Val, PyObject*& OutPyObj, const ESetErrorState SetErrorState = ESetErrorState::Yes);
 
 	/** uint8 overload */
-	bool Nativize(PyObject* PyObj, uint8& OutVal, const ESetErrorState SetErrorState = ESetErrorState::Yes);
-	bool Pythonize(const uint8 Val, PyObject*& OutPyObj, const ESetErrorState SetErrorState = ESetErrorState::Yes);
+	FPyConversionResult Nativize(PyObject* PyObj, uint8& OutVal, const ESetErrorState SetErrorState = ESetErrorState::Yes);
+	FPyConversionResult Pythonize(const uint8 Val, PyObject*& OutPyObj, const ESetErrorState SetErrorState = ESetErrorState::Yes);
 
 	/** int16 overload */
-	bool Nativize(PyObject* PyObj, int16& OutVal, const ESetErrorState SetErrorState = ESetErrorState::Yes);
-	bool Pythonize(const int16 Val, PyObject*& OutPyObj, const ESetErrorState SetErrorState = ESetErrorState::Yes);
+	FPyConversionResult Nativize(PyObject* PyObj, int16& OutVal, const ESetErrorState SetErrorState = ESetErrorState::Yes);
+	FPyConversionResult Pythonize(const int16 Val, PyObject*& OutPyObj, const ESetErrorState SetErrorState = ESetErrorState::Yes);
 
 	/** uint16 overload */
-	bool Nativize(PyObject* PyObj, uint16& OutVal, const ESetErrorState SetErrorState = ESetErrorState::Yes);
-	bool Pythonize(const uint16 Val, PyObject*& OutPyObj, const ESetErrorState SetErrorState = ESetErrorState::Yes);
+	FPyConversionResult Nativize(PyObject* PyObj, uint16& OutVal, const ESetErrorState SetErrorState = ESetErrorState::Yes);
+	FPyConversionResult Pythonize(const uint16 Val, PyObject*& OutPyObj, const ESetErrorState SetErrorState = ESetErrorState::Yes);
 
 	/** int32 overload */
-	bool Nativize(PyObject* PyObj, int32& OutVal, const ESetErrorState ErrorState = ESetErrorState::Yes);
-	bool Pythonize(const int32 Val, PyObject*& OutPyObj, const ESetErrorState SetErrorState = ESetErrorState::Yes);
+	FPyConversionResult Nativize(PyObject* PyObj, int32& OutVal, const ESetErrorState ErrorState = ESetErrorState::Yes);
+	FPyConversionResult Pythonize(const int32 Val, PyObject*& OutPyObj, const ESetErrorState SetErrorState = ESetErrorState::Yes);
 
 	/** uint32 overload */
-	bool Nativize(PyObject* PyObj, uint32& OutVal, const ESetErrorState SetErrorState = ESetErrorState::Yes);
-	bool Pythonize(const uint32 Val, PyObject*& OutPyObj, const ESetErrorState SetErrorState = ESetErrorState::Yes);
+	FPyConversionResult Nativize(PyObject* PyObj, uint32& OutVal, const ESetErrorState SetErrorState = ESetErrorState::Yes);
+	FPyConversionResult Pythonize(const uint32 Val, PyObject*& OutPyObj, const ESetErrorState SetErrorState = ESetErrorState::Yes);
 
 	/** int64 overload */
-	bool Nativize(PyObject* PyObj, int64& OutVal, const ESetErrorState SetErrorState = ESetErrorState::Yes);
-	bool Pythonize(const int64 Val, PyObject*& OutPyObj, const ESetErrorState SetErrorState = ESetErrorState::Yes);
+	FPyConversionResult Nativize(PyObject* PyObj, int64& OutVal, const ESetErrorState SetErrorState = ESetErrorState::Yes);
+	FPyConversionResult Pythonize(const int64 Val, PyObject*& OutPyObj, const ESetErrorState SetErrorState = ESetErrorState::Yes);
 
 	/** uint64 overload */
-	bool Nativize(PyObject* PyObj, uint64& OutVal, const ESetErrorState SetErrorState = ESetErrorState::Yes);
-	bool Pythonize(const uint64 Val, PyObject*& OutPyObj, const ESetErrorState SetErrorState = ESetErrorState::Yes);
+	FPyConversionResult Nativize(PyObject* PyObj, uint64& OutVal, const ESetErrorState SetErrorState = ESetErrorState::Yes);
+	FPyConversionResult Pythonize(const uint64 Val, PyObject*& OutPyObj, const ESetErrorState SetErrorState = ESetErrorState::Yes);
 
 	/** float overload */
-	bool Nativize(PyObject* PyObj, float& OutVal, const ESetErrorState SetErrorState = ESetErrorState::Yes);
-	bool Pythonize(const float Val, PyObject*& OutPyObj, const ESetErrorState SetErrorState = ESetErrorState::Yes);
+	FPyConversionResult Nativize(PyObject* PyObj, float& OutVal, const ESetErrorState SetErrorState = ESetErrorState::Yes);
+	FPyConversionResult Pythonize(const float Val, PyObject*& OutPyObj, const ESetErrorState SetErrorState = ESetErrorState::Yes);
 
 	/** double overload */
-	bool Nativize(PyObject* PyObj, double& OutVal, const ESetErrorState SetErrorState = ESetErrorState::Yes);
-	bool Pythonize(const double Val, PyObject*& OutPyObj, const ESetErrorState SetErrorState = ESetErrorState::Yes);
+	FPyConversionResult Nativize(PyObject* PyObj, double& OutVal, const ESetErrorState SetErrorState = ESetErrorState::Yes);
+	FPyConversionResult Pythonize(const double Val, PyObject*& OutPyObj, const ESetErrorState SetErrorState = ESetErrorState::Yes);
 
 	/** FString overload */
-	bool Nativize(PyObject* PyObj, FString& OutVal, const ESetErrorState SetErrorState = ESetErrorState::Yes);
-	bool Pythonize(const FString& Val, PyObject*& OutPyObj, const ESetErrorState SetErrorState = ESetErrorState::Yes);
+	FPyConversionResult Nativize(PyObject* PyObj, FString& OutVal, const ESetErrorState SetErrorState = ESetErrorState::Yes);
+	FPyConversionResult Pythonize(const FString& Val, PyObject*& OutPyObj, const ESetErrorState SetErrorState = ESetErrorState::Yes);
 
 	/** FName overload */
-	bool Nativize(PyObject* PyObj, FName& OutVal, const ESetErrorState SetErrorState = ESetErrorState::Yes);
-	bool Pythonize(const FName& Val, PyObject*& OutPyObj, const ESetErrorState SetErrorState = ESetErrorState::Yes);
+	FPyConversionResult Nativize(PyObject* PyObj, FName& OutVal, const ESetErrorState SetErrorState = ESetErrorState::Yes);
+	FPyConversionResult Pythonize(const FName& Val, PyObject*& OutPyObj, const ESetErrorState SetErrorState = ESetErrorState::Yes);
 
 	/** FText overload */
-	bool Nativize(PyObject* PyObj, FText& OutVal, const ESetErrorState SetErrorState = ESetErrorState::Yes);
-	bool Pythonize(const FText& Val, PyObject*& OutPyObj, const ESetErrorState SetErrorState = ESetErrorState::Yes);
+	FPyConversionResult Nativize(PyObject* PyObj, FText& OutVal, const ESetErrorState SetErrorState = ESetErrorState::Yes);
+	FPyConversionResult Pythonize(const FText& Val, PyObject*& OutPyObj, const ESetErrorState SetErrorState = ESetErrorState::Yes);
+
+	/** void* overload */
+	FPyConversionResult Nativize(PyObject* PyObj, void*& OutVal, const ESetErrorState SetErrorState = ESetErrorState::Yes);
+	FPyConversionResult Pythonize(void* Val, PyObject*& OutPyObj, const ESetErrorState SetErrorState = ESetErrorState::Yes);
 
 	/** UObject overload */
-	bool Nativize(PyObject* PyObj, UObject*& OutVal, const ESetErrorState SetErrorState = ESetErrorState::Yes);
-	bool Pythonize(UObject* Val, PyObject*& OutPyObj, const ESetErrorState SetErrorState = ESetErrorState::Yes);
+	FPyConversionResult Nativize(PyObject* PyObj, UObject*& OutVal, const ESetErrorState SetErrorState = ESetErrorState::Yes);
+	FPyConversionResult Pythonize(UObject* Val, PyObject*& OutPyObj, const ESetErrorState SetErrorState = ESetErrorState::Yes);
 
 	/** Conversion for object types, including optional type checking */
-	bool NativizeObject(PyObject* PyObj, UObject*& OutVal, UClass* ExpectedType, const ESetErrorState SetErrorState = ESetErrorState::Yes);
-	bool PythonizeObject(UObject* Val, PyObject*& OutPyObj, const ESetErrorState SetErrorState = ESetErrorState::Yes);
+	FPyConversionResult NativizeObject(PyObject* PyObj, UObject*& OutVal, UClass* ExpectedType, const ESetErrorState SetErrorState = ESetErrorState::Yes);
+	FPyConversionResult PythonizeObject(UObject* Val, PyObject*& OutPyObj, const ESetErrorState SetErrorState = ESetErrorState::Yes);
 	PyObject* PythonizeObject(UObject* Val, const ESetErrorState SetErrorState = ESetErrorState::Yes);
 
 	/** Conversion for class types, including optional type checking */
-	bool NativizeClass(PyObject* PyObj, UClass*& OutVal, UClass* ExpectedType, const ESetErrorState SetErrorState = ESetErrorState::Yes);
-	bool PythonizeClass(UClass* Val, PyObject*& OutPyObj, const ESetErrorState SetErrorState = ESetErrorState::Yes);
+	FPyConversionResult NativizeClass(PyObject* PyObj, UClass*& OutVal, UClass* ExpectedType, const ESetErrorState SetErrorState = ESetErrorState::Yes);
+	FPyConversionResult PythonizeClass(UClass* Val, PyObject*& OutPyObj, const ESetErrorState SetErrorState = ESetErrorState::Yes);
 	PyObject* PythonizeClass(UClass* Val, const ESetErrorState SetErrorState = ESetErrorState::Yes);
 
 	namespace Internal
 	{
 		/** Internal version of NativizeStruct/PythonizeStruct that work on the type-erased data */
-		bool NativizeStruct(PyObject* PyObj, UScriptStruct* StructType, void* StructInstance, const ESetErrorState SetErrorState);
-		bool PythonizeStruct(UScriptStruct* StructType, const void* StructInstance, PyObject*& OutPyObj, const ESetErrorState SetErrorState);
+		FPyConversionResult NativizeStruct(PyObject* PyObj, UScriptStruct* StructType, void* StructInstance, const ESetErrorState SetErrorState);
+		FPyConversionResult PythonizeStruct(UScriptStruct* StructType, const void* StructInstance, PyObject*& OutPyObj, const ESetErrorState SetErrorState);
 
 		/** Dummy catch-all for type conversions that aren't yet implemented */
 		template <typename T, typename Spec = void>
 		struct FTypeConv
 		{
-			static bool Nativize(PyObject* PyObj, T& OutVal, const ESetErrorState SetErrorState)
+			static FPyConversionResult Nativize(PyObject* PyObj, T& OutVal, const ESetErrorState SetErrorState)
 			{
 				ensureAlwaysMsgf(false, TEXT("Nativize not implemented for type"));
-				return false;
+				return FPyConversionResult::Failure();
 			}
 			
-			static bool Pythonize(const T& Val, PyObject*& OutPyObj, const ESetErrorState SetErrorState)
+			static FPyConversionResult Pythonize(const T& Val, PyObject*& OutPyObj, const ESetErrorState SetErrorState)
 			{
 				ensureAlwaysMsgf(false, TEXT("Pythonize not implemented for type"));
-				return false;
+				return FPyConversionResult::Failure();
 			}
 		};
 
@@ -117,12 +122,12 @@ namespace PyConversion
 		template <typename T>
 		struct FTypeConv<T, typename TEnableIf<TPointerIsConvertibleFromTo<typename TRemovePointer<T>::Type, UObject>::Value>::Type>
 		{
-			static bool Nativize(PyObject* PyObj, T& OutVal, const ESetErrorState SetErrorState)
+			static FPyConversionResult Nativize(PyObject* PyObj, T& OutVal, const ESetErrorState SetErrorState)
 			{
 				return PyConversion::NativizeObject(PyObj, (UObject*&)OutVal, TRemovePointer<T>::Type::StaticClass(), SetErrorState);
 			}
 
-			static bool Pythonize(const T& Val, PyObject*& OutPyObj, const ESetErrorState SetErrorState)
+			static FPyConversionResult Pythonize(const T& Val, PyObject*& OutPyObj, const ESetErrorState SetErrorState)
 			{
 				return PyConversion::PythonizeObject((UObject*)Val, OutPyObj, SetErrorState);
 			}
@@ -134,7 +139,7 @@ namespace PyConversion
 	 * Used to allow conversion from object and struct types that don't match a specific override (see FTypeConv).
 	 */
 	template <typename T>
-	bool Nativize(PyObject* PyObj, T& OutVal, const ESetErrorState SetErrorState = ESetErrorState::Yes)
+	FPyConversionResult Nativize(PyObject* PyObj, T& OutVal, const ESetErrorState SetErrorState = ESetErrorState::Yes)
 	{
 		return Internal::FTypeConv<T>::Nativize(PyObj, OutVal, SetErrorState);
 	}
@@ -144,7 +149,7 @@ namespace PyConversion
 	 * Used to allow conversion from object and struct types that don't match a specific override (see FTypeConv).
 	 */
 	template <typename T>
-	bool Pythonize(const T& Val, PyObject*& OutPyObj, const ESetErrorState SetErrorState = ESetErrorState::Yes)
+	FPyConversionResult Pythonize(const T& Val, PyObject*& OutPyObj, const ESetErrorState SetErrorState = ESetErrorState::Yes)
 	{
 		return Internal::FTypeConv<T>::Pythonize(Val, OutPyObj, SetErrorState);
 	}
@@ -160,14 +165,14 @@ namespace PyConversion
 
 	/** Conversion for known struct types */
 	template <typename T>
-	bool NativizeStruct(PyObject* PyObj, T& OutVal, const ESetErrorState SetErrorState = ESetErrorState::Yes)
+	FPyConversionResult NativizeStruct(PyObject* PyObj, T& OutVal, const ESetErrorState SetErrorState = ESetErrorState::Yes)
 	{
 		return Internal::NativizeStruct(PyObj, TBaseStructure<T>::Get(), &OutVal, SetErrorState);
 	}
 
 	/** Conversion for known struct types */
 	template <typename T>
-	bool PythonizeStruct(const T& Val, PyObject*& OutPyObj, const ESetErrorState SetErrorState = ESetErrorState::Yes)
+	FPyConversionResult PythonizeStruct(const T& Val, PyObject*& OutPyObj, const ESetErrorState SetErrorState = ESetErrorState::Yes)
 	{
 		return Internal::PythonizeStruct(TBaseStructure<T>::Get(), &Val, OutPyObj, SetErrorState);
 	}
@@ -181,17 +186,31 @@ namespace PyConversion
 		return Obj;
 	}
 
+	/** Conversion for known enum types */
+	template <typename T>
+	FPyConversionResult NativizeEnum(PyObject* PyObj, T& OutVal, const ESetErrorState SetErrorState = ESetErrorState::Yes)
+	{
+		return Nativize(PyObj, (__underlying_type(T)&)OutVal, SetErrorState);
+	}
+
+	/** Conversion for known enum types */
+	template <typename T>
+	FPyConversionResult PythonizeEnum(const T& Val, PyObject*& OutPyObj, const ESetErrorState SetErrorState = ESetErrorState::Yes)
+	{
+		return Pythonize((__underlying_type(T))Val, OutPyObj, SetErrorState);
+	}
+
 	/** Conversion for property instances (including fixed arrays) - ValueAddr should point to the property data */
-	bool NativizeProperty(PyObject* PyObj, const UProperty* Prop, void* ValueAddr, const FPyWrapperOwnerContext& InChangeOwner = FPyWrapperOwnerContext(), const ESetErrorState SetErrorState = ESetErrorState::Yes);
-	bool PythonizeProperty(const UProperty* Prop, const void* ValueAddr, PyObject*& OutPyObj, const EPyConversionMethod ConversionMethod = EPyConversionMethod::Copy, PyObject* OwnerPyObj = nullptr, const ESetErrorState SetErrorState = ESetErrorState::Yes);
+	FPyConversionResult NativizeProperty(PyObject* PyObj, const UProperty* Prop, void* ValueAddr, const FPyWrapperOwnerContext& InChangeOwner = FPyWrapperOwnerContext(), const ESetErrorState SetErrorState = ESetErrorState::Yes);
+	FPyConversionResult PythonizeProperty(const UProperty* Prop, const void* ValueAddr, PyObject*& OutPyObj, const EPyConversionMethod ConversionMethod = EPyConversionMethod::Copy, PyObject* OwnerPyObj = nullptr, const ESetErrorState SetErrorState = ESetErrorState::Yes);
 
 	/** Conversion for single property instances - ValueAddr should point to the property data */
-	bool NativizeProperty_Direct(PyObject* PyObj, const UProperty* Prop, void* ValueAddr, const FPyWrapperOwnerContext& InChangeOwner = FPyWrapperOwnerContext(), const ESetErrorState SetErrorState = ESetErrorState::Yes);
-	bool PythonizeProperty_Direct(const UProperty* Prop, const void* ValueAddr, PyObject*& OutPyObj, const EPyConversionMethod ConversionMethod = EPyConversionMethod::Copy, PyObject* OwnerPyObj = nullptr, const ESetErrorState SetErrorState = ESetErrorState::Yes);
+	FPyConversionResult NativizeProperty_Direct(PyObject* PyObj, const UProperty* Prop, void* ValueAddr, const FPyWrapperOwnerContext& InChangeOwner = FPyWrapperOwnerContext(), const ESetErrorState SetErrorState = ESetErrorState::Yes);
+	FPyConversionResult PythonizeProperty_Direct(const UProperty* Prop, const void* ValueAddr, PyObject*& OutPyObj, const EPyConversionMethod ConversionMethod = EPyConversionMethod::Copy, PyObject* OwnerPyObj = nullptr, const ESetErrorState SetErrorState = ESetErrorState::Yes);
 
 	/** Conversion for property instances within a structure (including fixed arrays) - BaseAddr should point to the structure data */
-	bool NativizeProperty_InContainer(PyObject* PyObj, const UProperty* Prop, void* BaseAddr, const int32 ArrayIndex, const FPyWrapperOwnerContext& InChangeOwner = FPyWrapperOwnerContext(), const ESetErrorState SetErrorState = ESetErrorState::Yes);
-	bool PythonizeProperty_InContainer(const UProperty* Prop, const void* BaseAddr, const int32 ArrayIndex, PyObject*& OutPyObj, const EPyConversionMethod ConversionMethod = EPyConversionMethod::Copy, PyObject* OwnerPyObj = nullptr, const ESetErrorState SetErrorState = ESetErrorState::Yes);
+	FPyConversionResult NativizeProperty_InContainer(PyObject* PyObj, const UProperty* Prop, void* BaseAddr, const int32 ArrayIndex, const FPyWrapperOwnerContext& InChangeOwner = FPyWrapperOwnerContext(), const ESetErrorState SetErrorState = ESetErrorState::Yes);
+	FPyConversionResult PythonizeProperty_InContainer(const UProperty* Prop, const void* BaseAddr, const int32 ArrayIndex, PyObject*& OutPyObj, const EPyConversionMethod ConversionMethod = EPyConversionMethod::Copy, PyObject* OwnerPyObj = nullptr, const ESetErrorState SetErrorState = ESetErrorState::Yes);
 
 	/**
 	 * Helper function used to emit property change notifications as value changes are made
diff --git a/Engine/Plugins/Experimental/PythonScriptPlugin/Source/PythonScriptPlugin/Private/PyCore.cpp b/Engine/Plugins/Experimental/PythonScriptPlugin/Source/PythonScriptPlugin/Private/PyCore.cpp
--- a/Engine/Plugins/Experimental/PythonScriptPlugin/Source/PythonScriptPlugin/Private/PyCore.cpp
+++ b/Engine/Plugins/Experimental/PythonScriptPlugin/Source/PythonScriptPlugin/Private/PyCore.cpp
@@ -90,6 +90,33 @@ UObject* GetPythonTypeContainer()
 	return GPythonTypeContainer.Get();
 }
 
+
+FPyWrapperDelegateHandle* FPyWrapperDelegateHandle::CreateInstance(const FDelegateHandle& InValue)
+{
+	FPyWrapperDelegateHandlePtr NewInstance = FPyWrapperDelegateHandlePtr::StealReference(FPyWrapperDelegateHandle::New(&PyDelegateHandleType));
+	if (NewInstance)
+	{
+		if (FPyWrapperDelegateHandle::Init(NewInstance, InValue) != 0)
+		{
+			PyUtil::LogPythonError();
+			return nullptr;
+		}
+	}
+	return NewInstance.Release();
+}
+
+FPyWrapperDelegateHandle* FPyWrapperDelegateHandle::CastPyObject(PyObject* InPyObject)
+{
+	if (PyObject_IsInstance(InPyObject, (PyObject*)&PyDelegateHandleType) == 1)
+	{
+		Py_INCREF(InPyObject);
+		return (FPyWrapperDelegateHandle*)InPyObject;
+	}
+
+	return nullptr;
+}
+
+
 FPyUValueDef* FPyUValueDef::New(PyTypeObject* InType)
 {
 	FPyUValueDef* Self = (FPyUValueDef*)InType->tp_alloc(InType, 0);
@@ -453,6 +480,7 @@ PyTypeObject InitializePyUFunctionDefType()
 }
 
 
+PyTypeObject PyDelegateHandleType = InitializePyWrapperBasicType<FPyWrapperDelegateHandle>("_DelegateHandle", "Type for all UE4 exposed FDelegateHandle instances");
 PyTypeObject PyUValueDefType = InitializePyUValueDefType();
 PyTypeObject PyUPropertyDefType = InitializePyUPropertyDefType();
 PyTypeObject PyUFunctionDefType = InitializePyUFunctionDefType();
@@ -536,6 +564,39 @@ PyObject* Reload(PyObject* InSelf, PyObject* InArgs)
 	Py_RETURN_NONE;
 }
 
+PyObject* LoadModule(PyObject* InSelf, PyObject* InArgs)
+{
+	PyObject* PyObj = nullptr;
+	if (!PyArg_ParseTuple(InArgs, "O:load_module", &PyObj))
+	{
+		return nullptr;
+	}
+
+	FString ModuleName;
+	if (!PyConversion::Nativize(PyObj, ModuleName))
+	{
+		return nullptr;
+	}
+
+	if (!FModuleManager::Get().ModuleExists(*ModuleName))
+	{
+		PyUtil::SetPythonError(PyExc_KeyError, TEXT("load_module"), *FString::Printf(TEXT("'%s' isn't a known module name"), *ModuleName));
+		return nullptr;
+	}
+
+	IModuleInterface* ModulePtr = FModuleManager::Get().LoadModule(*ModuleName);
+	if (ModulePtr)
+	{
+		FPythonScriptPlugin* PythonScriptPlugin = FPythonScriptPlugin::Get();
+		if (PythonScriptPlugin)
+		{
+			PythonScriptPlugin->ImportUnrealModule(*ModuleName);
+		}
+	}
+
+	Py_RETURN_NONE;
+}
+
 PyObject* FindOrLoadObjectImpl(PyObject* InSelf, PyObject* InArgs, PyObject* InKwds, const TCHAR* InFuncName, const TFunctionRef<UObject*(UClass*, UObject*, const TCHAR*)>& InFunc)
 {
 	PyObject* PyOuterObj = nullptr;
@@ -892,6 +953,7 @@ PyMethodDef PyCoreMethods[] = {
 	{ "log_error", PyCFunctionCast(&LogError), METH_VARARGS, "x.log_error(str) -- log the given argument as an error in the LogPython category" },
 	{ "log_flush", PyCFunctionCast(&LogFlush), METH_NOARGS, "x.log_flush() -- flush the log to disk" },
 	{ "reload", PyCFunctionCast(&Reload), METH_VARARGS, "x.reload(str) -- reload the given Unreal Python module" },
+	{ "load_module", PyCFunctionCast(&LoadModule), METH_VARARGS, "x.load_module(str) -- load the given Unreal module and generate any Python code for its reflected types" },
 	{ "find_object", PyCFunctionCast(&FindObject), METH_VARARGS | METH_KEYWORDS, "x.find_object(outer, name, type=Object) -> Object -- find an Unreal object with the given outer and name, optionally validating its type" },
 	{ "load_object", PyCFunctionCast(&LoadObject), METH_VARARGS | METH_KEYWORDS, "x.load_object(outer, name, type=Object) -> Object -- load an Unreal object with the given outer and name, optionally validating its type" },
 	{ "load_class", PyCFunctionCast(&LoadClass), METH_VARARGS | METH_KEYWORDS, "x.load_class(outer, name, type=Object) -> Class -- load an Unreal class with the given outer and name, optionally validating its base type" },
@@ -922,6 +984,8 @@ void InitializeModule()
 	PyObject* PyModule = Py_InitModule("_unreal_core", PyCoreMethods);
 #endif	// PY_MAJOR_VERSION >= 3
 
+	PyType_Ready(&PyDelegateHandleType);
+
 	if (PyType_Ready(&PyUValueDefType) == 0)
 	{
 		Py_INCREF(&PyUValueDefType);
diff --git a/Engine/Plugins/Experimental/PythonScriptPlugin/Source/PythonScriptPlugin/Private/PyCore.h b/Engine/Plugins/Experimental/PythonScriptPlugin/Source/PythonScriptPlugin/Private/PyCore.h
--- a/Engine/Plugins/Experimental/PythonScriptPlugin/Source/PythonScriptPlugin/Private/PyCore.h
+++ b/Engine/Plugins/Experimental/PythonScriptPlugin/Source/PythonScriptPlugin/Private/PyCore.h
@@ -3,6 +3,7 @@
 #pragma once
 
 #include "IncludePython.h"
+#include "PyWrapperBasic.h"
 #include "PyPtr.h"
 #include "CoreMinimal.h"
 #include "Misc/EnumClassFlags.h"
@@ -15,6 +16,9 @@ UObject* GetPythonPropertyContainer();
 /** Get the object that Python created types should be outered to */
 UObject* GetPythonTypeContainer();
 
+/** Python type for FPyDelegateHandle */
+extern PyTypeObject PyDelegateHandleType;
+
 /** Python type for FPyUValueDef */
 extern PyTypeObject PyUValueDefType;
 
@@ -24,6 +28,20 @@ extern PyTypeObject PyUPropertyDefType;
 /** Python type for FPyUFunctionDef */
 extern PyTypeObject PyUFunctionDefType;
 
+/** Type for all UE4 exposed FDelegateHandle instances */
+struct FPyWrapperDelegateHandle : public TPyWrapperBasic<FDelegateHandle, FPyWrapperDelegateHandle>
+{
+	typedef TPyWrapperBasic<FDelegateHandle, FPyWrapperDelegateHandle> Super;
+
+	/** Create and initialize a new wrapper instance from the given native instance */
+	static FPyWrapperDelegateHandle* CreateInstance(const FDelegateHandle& InValue);
+
+	/** Cast the given Python object to this wrapped type (returns a new reference) */
+	static FPyWrapperDelegateHandle* CastPyObject(PyObject* InPyObject);
+};
+
+typedef TPyPtr<FPyWrapperDelegateHandle> FPyWrapperDelegateHandlePtr;
+
 /** Type used to define constant values from Python */
 struct FPyUValueDef
 {
diff --git a/Engine/Plugins/Experimental/PythonScriptPlugin/Source/PythonScriptPlugin/Private/PyGenUtil.cpp b/Engine/Plugins/Experimental/PythonScriptPlugin/Source/PythonScriptPlugin/Private/PyGenUtil.cpp
--- a/Engine/Plugins/Experimental/PythonScriptPlugin/Source/PythonScriptPlugin/Private/PyGenUtil.cpp
+++ b/Engine/Plugins/Experimental/PythonScriptPlugin/Source/PythonScriptPlugin/Private/PyGenUtil.cpp
@@ -2,13 +2,15 @@
 
 #include "PyGenUtil.h"
 #include "PyUtil.h"
+#include "PyConversion.h"
 #include "PyWrapperBase.h"
+#include "PyWrapperStruct.h"
 #include "Internationalization/BreakIterator.h"
-#include "UObject/UnrealType.h"
 #include "UObject/Class.h"
 #include "UObject/EnumProperty.h"
 #include "UObject/Package.h"
 #include "UObject/TextProperty.h"
+#include "UObject/PropertyPortFlags.h"
 #include "UObject/UnrealType.h"
 #include "Engine/BlueprintGeneratedClass.h"
 #include "Engine/UserDefinedStruct.h"
@@ -19,19 +21,56 @@
 namespace PyGenUtil
 {
 
+/** Case sensitive hashing function for TSet */
+struct FCaseSensitiveStringSetFuncs : BaseKeyFuncs<FString, FString>
+{
+	static FORCEINLINE const FString& GetSetKey(const FString& Element)
+	{
+		return Element;
+	}
+	static FORCEINLINE bool Matches(const FString& A, const FString& B)
+	{
+		return A.Equals(B, ESearchCase::CaseSensitive);
+	}
+	static FORCEINLINE uint32 GetKeyHash(const FString& Key)
+	{
+		return FCrc::StrCrc32<TCHAR>(*Key);
+	}
+};
+
+
 const FName ScriptNameMetaDataKey = TEXT("ScriptName");
 const FName ScriptNoExportMetaDataKey = TEXT("ScriptNoExport");
+const FName ScriptMethodMetaDataKey = TEXT("ScriptMethod");
+const FName ScriptMathOpMetaDataKey = TEXT("ScriptMathOp");
 const FName BlueprintTypeMetaDataKey = TEXT("BlueprintType");
 const FName NotBlueprintTypeMetaDataKey = TEXT("NotBlueprintType");
 const FName BlueprintSpawnableComponentMetaDataKey = TEXT("BlueprintSpawnableComponent");
 const FName BlueprintGetterMetaDataKey = TEXT("BlueprintGetter");
 const FName BlueprintSetterMetaDataKey = TEXT("BlueprintSetter");
+const FName DeprecatedPropertyMetaDataKey = TEXT("DeprecatedProperty");
+const FName DeprecatedFunctionMetaDataKey = TEXT("DeprecatedFunction");
+const FName DeprecationMessageMetaDataKey = TEXT("DeprecationMessage");
+const FName CustomStructureParamMetaDataKey = TEXT("CustomStructureParam");
 const TCHAR* HiddenMetaDataKey = TEXT("Hidden");
 
 
 TSharedPtr<IBreakIterator> NameBreakIterator;
 
 
+void FGeneratedWrappedFunction::SetFunctionAndExtractParams(const UFunction* InFunc)
+{
+	Func = InFunc;
+	InputParams.Reset();
+	OutputParams.Reset();
+
+	if (Func)
+	{
+		ExtractFunctionParams(Func, InputParams, OutputParams);
+	}
+}
+
+
 void FGeneratedWrappedMethod::ToPython(FPyMethodWithClosureDef& OutPyMethod) const
 {
 	OutPyMethod.MethodName = MethodName.GetData();
@@ -42,6 +81,84 @@ void FGeneratedWrappedMethod::ToPython(FPyMethodWithClosureDef& OutPyMethod) con
 }
 
 
+void FGeneratedWrappedMethods::Finalize()
+{
+	check(PyMethods.Num() == 0);
+
+	PyMethods.Reserve(TypeMethods.Num() + 1);
+	for (const FGeneratedWrappedMethod& TypeMethod : TypeMethods)
+	{
+		FPyMethodWithClosureDef& PyMethod = PyMethods[PyMethods.AddZeroed()];
+		TypeMethod.ToPython(PyMethod);
+	}
+	PyMethods.AddZeroed(); // null terminator
+}
+
+
+void FGeneratedWrappedDynamicStructMethodWithClosure::Finalize()
+{
+	ToPython(PyMethod);
+}
+
+
+bool FGeneratedWrappedStructMathOpFunction::SetFunctionAndExtractParams(const UFunction* InFunc)
+{
+	FGeneratedWrappedFunction::SetFunctionAndExtractParams(InFunc);
+
+	// The struct parameter should be the first parameter
+	if (InputParams.Num() > 0 && InputParams[0].ParamProp->IsA<UStructProperty>())
+	{
+		StructParam = InputParams[0];
+		InputParams.RemoveAt(0, 1, /*bAllowShrinking*/false);
+	}
+
+	return StructParam.ParamProp != nullptr;
+}
+
+
+bool FGeneratedWrappedStructMathOpStack::StringToOpType(const TCHAR* InStr, EOpType& OutOpType)
+{
+	const TCHAR* OpStackStrings[(int32)EOpType::Num] = {
+		TEXT("+"),
+		TEXT("+="),
+		TEXT("-"),
+		TEXT("-="),
+		TEXT("*"),
+		TEXT("*="),
+		TEXT("/"),
+		TEXT("/="),
+		TEXT("%"),
+		TEXT("%="),
+		TEXT("&"),
+		TEXT("&="),
+		TEXT("|"),
+		TEXT("|="),
+		TEXT("^"),
+		TEXT("^="),
+		TEXT(">>"),
+		TEXT(">>="),
+		TEXT("<<"),
+		TEXT("<<="),
+	};
+
+	for (int32 OpTypeIndex = 0; OpTypeIndex < (int32)EOpType::Num; ++OpTypeIndex)
+	{
+		if (FCString::Strcmp(InStr, OpStackStrings[OpTypeIndex]) == 0)
+		{
+			OutOpType = (EOpType)OpTypeIndex;
+			return true;
+		}
+	}
+
+	return false;
+}
+
+bool FGeneratedWrappedStructMathOpStack::IsInlineOp(const EOpType InOpType)
+{
+	return (int32)InOpType % 2 != 0;
+}
+
+
 void FGeneratedWrappedGetSet::ToPython(PyGetSetDef& OutPyGetSet) const
 {
 	OutPyGetSet.name = (char*)GetSetName.GetData();
@@ -52,6 +169,20 @@ void FGeneratedWrappedGetSet::ToPython(PyGetSetDef& OutPyGetSet) const
 }
 
 
+void FGeneratedWrappedGetSets::Finalize()
+{
+	check(PyGetSets.Num() == 0);
+
+	PyGetSets.Reserve(TypeGetSets.Num() + 1);
+	for (const FGeneratedWrappedGetSet& TypeGetSet : TypeGetSets)
+	{
+		PyGetSetDef& PyGetSet = PyGetSets[PyGetSets.AddZeroed()];
+		TypeGetSet.ToPython(PyGetSet);
+	}
+	PyGetSets.AddZeroed(); // null terminator
+}
+
+
 FGeneratedWrappedPropertyDoc::FGeneratedWrappedPropertyDoc(const UProperty* InProp)
 {
 	PythonPropName = PyGenUtil::GetPropertyPythonName(InProp);
@@ -109,35 +240,139 @@ void FGeneratedWrappedPropertyDoc::AppendDocString(const TArray<FGeneratedWrappe
 
 
 bool FGeneratedWrappedType::Finalize()
+{
+	Finalize_PreReady();
+
+	if (PyType_Ready(&PyType) == 0)
+	{
+		Finalize_PostReady();
+		FPyWrapperBaseMetaData::SetMetaData(&PyType, MetaData.Get());
+		return true;
+	}
+
+	return false;
+}
+
+void FGeneratedWrappedType::Finalize_PreReady()
 {
 	PyType.tp_name = TypeName.GetData();
 	PyType.tp_doc = TypeDoc.GetData();
+}
 
-	PyMethods.Reserve(TypeMethods.Num() + 1);
-	for (const FGeneratedWrappedMethod& TypeMethod : TypeMethods)
-	{
-		FPyMethodWithClosureDef& PyMethod = PyMethods[PyMethods.AddZeroed()];
-		TypeMethod.ToPython(PyMethod);
+void FGeneratedWrappedType::Finalize_PostReady()
+{
+}
+
+
+void FGeneratedWrappedStructType::Finalize_PreReady()
+{
+	FGeneratedWrappedType::Finalize_PreReady();
+
+	GetSets.Finalize();
+	PyType.tp_getset = GetSets.PyGetSets.GetData();
+
+#if PY_MAJOR_VERSION < 3
+	PyType.tp_flags |= Py_TPFLAGS_CHECKTYPES;
+#endif	// PY_MAJOR_VERSION < 3
+
+#define DEFINE_BINARY_MATH_FUNC(OP)																									\
+	static PyObject* OP(FPyWrapperStruct* InLHS, PyObject* InRHS)																	\
+	{																																\
+		return FPyWrapperStruct::CallBinaryOperator_Impl(InLHS, InRHS, PyGenUtil::FGeneratedWrappedStructMathOpStack::EOpType::OP);	\
 	}
-	PyMethods.AddZeroed(); // null terminator
-	
-	PyMethods.Reserve(PyGetSets.Num() + 1);
-	for (const FGeneratedWrappedGetSet& TypeGetSet : TypeGetSets)
+	struct FNumberFuncs
 	{
-		PyGetSetDef& PyGetSet = PyGetSets[PyGetSets.AddZeroed()];
-		TypeGetSet.ToPython(PyGetSet);
-	}
-	PyGetSets.AddZeroed(); // null terminator
-	PyType.tp_getset = PyGetSets.GetData();
+		DEFINE_BINARY_MATH_FUNC(Add)
+		DEFINE_BINARY_MATH_FUNC(InlineAdd)
+		DEFINE_BINARY_MATH_FUNC(Subtract)
+		DEFINE_BINARY_MATH_FUNC(InlineSubtract)
+		DEFINE_BINARY_MATH_FUNC(Multiply)
+		DEFINE_BINARY_MATH_FUNC(InlineMultiply)
+		DEFINE_BINARY_MATH_FUNC(Divide)
+		DEFINE_BINARY_MATH_FUNC(InlineDivide)
+		DEFINE_BINARY_MATH_FUNC(Modulus)
+		DEFINE_BINARY_MATH_FUNC(InlineModulus)
+		DEFINE_BINARY_MATH_FUNC(And)
+		DEFINE_BINARY_MATH_FUNC(InlineAnd)
+		DEFINE_BINARY_MATH_FUNC(Or)
+		DEFINE_BINARY_MATH_FUNC(InlineOr)
+		DEFINE_BINARY_MATH_FUNC(Xor)
+		DEFINE_BINARY_MATH_FUNC(InlineXor)
+		DEFINE_BINARY_MATH_FUNC(RightShift)
+		DEFINE_BINARY_MATH_FUNC(InlineRightShift)
+		DEFINE_BINARY_MATH_FUNC(LeftShift)
+		DEFINE_BINARY_MATH_FUNC(InlineLeftShift)
+	};
+#undef DEFINE_BINARY_MATH_FUNC
+
+	PyNumber.nb_add = (binaryfunc)&FNumberFuncs::Add;
+	PyNumber.nb_inplace_add = (binaryfunc)&FNumberFuncs::InlineAdd;
+	PyNumber.nb_subtract = (binaryfunc)&FNumberFuncs::Subtract;
+	PyNumber.nb_inplace_subtract = (binaryfunc)&FNumberFuncs::InlineSubtract;
+	PyNumber.nb_multiply = (binaryfunc)&FNumberFuncs::Multiply;
+	PyNumber.nb_inplace_multiply = (binaryfunc)&FNumberFuncs::InlineMultiply;
+#if PY_MAJOR_VERSION >= 3
+	PyNumber.nb_true_divide = (binaryfunc)&FNumberFuncs::Divide;
+	PyNumber.nb_inplace_true_divide = (binaryfunc)&FNumberFuncs::InlineDivide;
+#else	// PY_MAJOR_VERSION >= 3
+	PyNumber.nb_divide = (binaryfunc)&FNumberFuncs::Divide;
+	PyNumber.nb_inplace_divide = (binaryfunc)&FNumberFuncs::InlineDivide;
+#endif	// PY_MAJOR_VERSION >= 3
+	PyNumber.nb_and = (binaryfunc)&FNumberFuncs::And;
+	PyNumber.nb_inplace_and = (binaryfunc)&FNumberFuncs::InlineAnd;
+	PyNumber.nb_or = (binaryfunc)&FNumberFuncs::Or;
+	PyNumber.nb_inplace_or = (binaryfunc)&FNumberFuncs::InlineOr;
+	PyNumber.nb_xor = (binaryfunc)&FNumberFuncs::Xor;
+	PyNumber.nb_inplace_xor = (binaryfunc)&FNumberFuncs::InlineXor;
+	PyNumber.nb_rshift = (binaryfunc)&FNumberFuncs::RightShift;
+	PyNumber.nb_inplace_rshift = (binaryfunc)&FNumberFuncs::InlineRightShift;
+	PyNumber.nb_lshift = (binaryfunc)&FNumberFuncs::LeftShift;
+	PyNumber.nb_inplace_lshift = (binaryfunc)&FNumberFuncs::InlineLeftShift;
+
+	PyType.tp_as_number = &PyNumber;
+}
 
-	if (PyType_Ready(&PyType) == 0)
+void FGeneratedWrappedStructType::AddDynamicStructMethod(FGeneratedWrappedDynamicStructMethod&& InDynamicStructMethod)
+{
+	TSharedRef<FGeneratedWrappedDynamicStructMethodWithClosure> DynamicStructMethod = DynamicStructMethods.Add_GetRef(MakeShared<FGeneratedWrappedDynamicStructMethodWithClosure>());
+	static_cast<FGeneratedWrappedDynamicStructMethod&>(*DynamicStructMethod) = MoveTemp(InDynamicStructMethod);
+	DynamicStructMethod->Finalize();
+	FPyMethodWithClosureDef::AddMethod(&DynamicStructMethod->PyMethod, &PyType);
+}
+
+
+void FGeneratedWrappedClassType::Finalize_PreReady()
+{
+	FGeneratedWrappedType::Finalize_PreReady();
+
+	Methods.Finalize();
+
+	GetSets.Finalize();
+	PyType.tp_getset = GetSets.PyGetSets.GetData();
+}
+
+void FGeneratedWrappedClassType::Finalize_PostReady()
+{
+	FGeneratedWrappedType::Finalize_PostReady();
+
+	FPyMethodWithClosureDef::AddMethods(Methods.PyMethods.GetData(), &PyType);
+}
+
+
+FPythonizeTooltipContext::FPythonizeTooltipContext(const UProperty* InProp, const UFunction* InFunc, const uint64 InReadOnlyFlags)
+	: Prop(InProp)
+	, Func(InFunc)
+	, ReadOnlyFlags(InReadOnlyFlags)
+{
+	if (Prop)
 	{
-		FPyMethodWithClosureDef::AddMethods(PyMethods.GetData(), &PyType);
-		FPyWrapperBaseMetaData::SetMetaData(&PyType, MetaData.Get());
-		return true;
+		PyGenUtil::IsDeprecatedProperty(Prop, &DeprecationMessage);
 	}
 
-	return false;
+	if (Func)
+	{
+		PyGenUtil::IsDeprecatedFunction(Func, &DeprecationMessage);
+	}
 }
 
 
@@ -190,14 +425,63 @@ PyObject* GetPostInitFunc(PyTypeObject* InPyType)
 	return PostInitFunc.Release();
 }
 
-void AddStructInitParam(const TCHAR* InUnrealPropName, const TCHAR* InPythonAttrName, TArray<FGeneratedWrappedMethodParameter>& OutInitParams)
+void AddStructInitParam(const UProperty* InUnrealProp, const TCHAR* InPythonAttrName, TArray<FGeneratedWrappedMethodParameter>& OutInitParams)
 {
-	FGeneratedWrappedMethodParameter& InitParam = OutInitParams[OutInitParams.AddDefaulted()];
+	FGeneratedWrappedMethodParameter& InitParam = OutInitParams.AddDefaulted_GetRef();
 	InitParam.ParamName = TCHARToUTF8Buffer(InPythonAttrName);
-	InitParam.ParamPropName = InUnrealPropName;
+	InitParam.ParamProp = InUnrealProp;
 	InitParam.ParamDefaultValue = FString();
 }
 
+void ExtractFunctionParams(const UFunction* InFunc, TArray<FGeneratedWrappedMethodParameter>& OutInputParams, TArray<FGeneratedWrappedMethodParameter>& OutOutputParams)
+{
+	auto AddGeneratedWrappedMethodParameter = [InFunc](const UProperty* InParam, TArray<FGeneratedWrappedMethodParameter>& OutParams)
+	{
+		const FString ParamName = InParam->GetName();
+		const FString PythonParamName = PyGenUtil::PythonizePropertyName(ParamName, PyGenUtil::EPythonizeNameCase::Lower);
+		const FName DefaultValueMetaDataKey = *FString::Printf(TEXT("CPP_Default_%s"), *ParamName);
+
+		PyGenUtil::FGeneratedWrappedMethodParameter& GeneratedWrappedMethodParam = OutParams.AddDefaulted_GetRef();
+		GeneratedWrappedMethodParam.ParamName = PyGenUtil::TCHARToUTF8Buffer(*PythonParamName);
+		GeneratedWrappedMethodParam.ParamProp = InParam;
+		if (InFunc->HasMetaData(DefaultValueMetaDataKey))
+		{
+			GeneratedWrappedMethodParam.ParamDefaultValue = InFunc->GetMetaData(DefaultValueMetaDataKey);
+		}
+	};
+
+	if (const UProperty* ReturnProp = InFunc->GetReturnProperty())
+	{
+		AddGeneratedWrappedMethodParameter(ReturnProp, OutOutputParams);
+	}
+
+	for (TFieldIterator<const UProperty> ParamIt(InFunc); ParamIt; ++ParamIt)
+	{
+		const UProperty* Param = *ParamIt;
+
+		if (PyUtil::IsInputParameter(Param))
+		{
+			AddGeneratedWrappedMethodParameter(Param, OutInputParams);
+		}
+
+		if (PyUtil::IsOutputParameter(Param))
+		{
+			AddGeneratedWrappedMethodParameter(Param, OutOutputParams);
+		}
+	}
+}
+
+void ApplyParamDefaults(void* InBaseParamsAddr, const TArray<FGeneratedWrappedMethodParameter>& InParamDef)
+{
+	for (const FGeneratedWrappedMethodParameter& ParamDef : InParamDef)
+	{
+		if (ParamDef.ParamDefaultValue.IsSet())
+		{
+			ParamDef.ParamProp->ImportText(*ParamDef.ParamDefaultValue.GetValue(), ParamDef.ParamProp->ContainerPtrToValuePtr<void>(InBaseParamsAddr), PPF_None, nullptr);
+		}
+	}
+}
+
 bool ParseMethodParameters(PyObject* InArgs, PyObject* InKwds, const TArray<FGeneratedWrappedMethodParameter>& InParamDef, const char* InPyMethodName, TArray<PyObject*>& OutPyParams)
 {
 	if (!InArgs || !PyTuple_Check(InArgs) || (InKwds && !PyDict_Check(InKwds)) || !InPyMethodName)
@@ -279,6 +563,184 @@ bool ParseMethodParameters(PyObject* InArgs, PyObject* InKwds, const TArray<FGen
 	return true;
 }
 
+PyObject* PackReturnValues(void* InBaseParamsAddr, const TArray<FGeneratedWrappedMethodParameter>& InOutputParams, const TCHAR* InErrorCtxt, const TCHAR* InCallingCtxt)
+{
+	if (!InOutputParams.Num())
+	{
+		Py_RETURN_NONE;
+	}
+
+	int32 ReturnPropIndex = 0;
+
+	// If we have multiple return values and the main return value is a bool, we return None (for false) or the (potentially packed) return value without the bool (for true)
+	if (InOutputParams.Num() > 1 && InOutputParams[0].ParamProp->HasAnyPropertyFlags(CPF_ReturnParm) && InOutputParams[0].ParamProp->IsA<UBoolProperty>())
+	{
+		const UBoolProperty* BoolReturn = CastChecked<const UBoolProperty>(InOutputParams[0].ParamProp);
+		const bool bReturnValue = BoolReturn->GetPropertyValue(BoolReturn->ContainerPtrToValuePtr<void>(InBaseParamsAddr));
+		if (!bReturnValue)
+		{
+			Py_RETURN_NONE;
+		}
+
+		ReturnPropIndex = 1; // Start packing at the 1st out value
+	}
+
+	// Do we need to return a packed tuple, or just a single value?
+	const int32 NumPropertiesToPack = InOutputParams.Num() - ReturnPropIndex;
+	if (NumPropertiesToPack == 1)
+	{
+		PyObject* OutParamPyObj = nullptr;
+		if (!PyConversion::PythonizeProperty_InContainer(InOutputParams[ReturnPropIndex].ParamProp, InBaseParamsAddr, 0, OutParamPyObj, EPyConversionMethod::Steal))
+		{
+			PyUtil::SetPythonError(PyExc_TypeError, InErrorCtxt, *FString::Printf(TEXT("Failed to convert return property '%s' (%s) when calling %s"), *InOutputParams[ReturnPropIndex].ParamProp->GetName(), *InOutputParams[ReturnPropIndex].ParamProp->GetClass()->GetName(), InCallingCtxt));
+			return nullptr;
+		}
+		return OutParamPyObj;
+	}
+	else
+	{
+		int32 OutParamTupleIndex = 0;
+		FPyObjectPtr OutParamTuple = FPyObjectPtr::StealReference(PyTuple_New(NumPropertiesToPack));
+		for (; ReturnPropIndex < InOutputParams.Num(); ++ReturnPropIndex)
+		{
+			PyObject* OutParamPyObj = nullptr;
+			if (!PyConversion::PythonizeProperty_InContainer(InOutputParams[ReturnPropIndex].ParamProp, InBaseParamsAddr, 0, OutParamPyObj, EPyConversionMethod::Steal))
+			{
+				PyUtil::SetPythonError(PyExc_TypeError, InErrorCtxt, *FString::Printf(TEXT("Failed to convert return property '%s' (%s) when calling function %s"), *InOutputParams[ReturnPropIndex].ParamProp->GetName(), *InOutputParams[ReturnPropIndex].ParamProp->GetClass()->GetName(), InCallingCtxt));
+				return nullptr;
+			}
+			PyTuple_SetItem(OutParamTuple, OutParamTupleIndex++, OutParamPyObj); // SetItem steals the reference
+		}
+		return OutParamTuple.Release();
+	}
+}
+
+bool UnpackReturnValues(PyObject* InRetVals, void* InBaseParamsAddr, const TArray<FGeneratedWrappedMethodParameter>& InOutputParams, const TCHAR* InErrorCtxt, const TCHAR* InCallingCtxt)
+{
+	if (!InOutputParams.Num())
+	{
+		return true;
+	}
+
+	int32 ReturnPropIndex = 0;
+
+	// If we have multiple return values and the main return value is a bool, we expect None (for false) or the (potentially packed) return value without the bool (for true)
+	if (InOutputParams.Num() > 1 && InOutputParams[0].ParamProp->HasAnyPropertyFlags(CPF_ReturnParm) && InOutputParams[0].ParamProp->IsA<UBoolProperty>())
+	{
+		const UBoolProperty* BoolReturn = CastChecked<const UBoolProperty>(InOutputParams[0].ParamProp);
+		const bool bReturnValue = InRetVals != Py_None;
+		BoolReturn->SetPropertyValue(BoolReturn->ContainerPtrToValuePtr<void>(InBaseParamsAddr), bReturnValue);
+
+		ReturnPropIndex = 1; // Start unpacking at the 1st out value
+	}
+
+	// Do we need to expect a packed tuple, or just a single value?
+	const int32 NumPropertiesToUnpack = InOutputParams.Num() - ReturnPropIndex;
+	if (NumPropertiesToUnpack == 1)
+	{
+		if (!PyConversion::NativizeProperty_InContainer(InRetVals, InOutputParams[ReturnPropIndex].ParamProp, InBaseParamsAddr, 0))
+		{
+			PyUtil::SetPythonError(PyExc_TypeError, InErrorCtxt, *FString::Printf(TEXT("Failed to convert return property '%s' (%s) when calling %s"), *InOutputParams[ReturnPropIndex].ParamProp->GetName(), *InOutputParams[ReturnPropIndex].ParamProp->GetClass()->GetName(), InCallingCtxt));
+			return false;
+		}
+	}
+	else
+	{
+		if (!PyTuple_Check(InRetVals))
+		{
+			PyUtil::SetPythonError(PyExc_TypeError, InErrorCtxt, *FString::Printf(TEXT("Expected a 'tuple' return type, but got '%s' when calling %s"), *PyUtil::GetFriendlyTypename(InRetVals), InCallingCtxt));
+			return false;
+		}
+
+		const int32 RetTupleSize = PyTuple_Size(InRetVals);
+		if (RetTupleSize != NumPropertiesToUnpack)
+		{
+			PyUtil::SetPythonError(PyExc_TypeError, InErrorCtxt, *FString::Printf(TEXT("Expected a 'tuple' return type containing '%d' items but got one containing '%d' items when calling %s"), NumPropertiesToUnpack, RetTupleSize, InCallingCtxt));
+			return false;
+		}
+
+		int32 RetTupleIndex = 0;
+		for (; ReturnPropIndex < InOutputParams.Num(); ++ReturnPropIndex)
+		{
+			PyObject* RetVal = PyTuple_GetItem(InRetVals, RetTupleIndex++);
+			if (!PyConversion::NativizeProperty_InContainer(RetVal, InOutputParams[ReturnPropIndex].ParamProp, InBaseParamsAddr, 0))
+			{
+				PyUtil::SetPythonError(PyExc_TypeError, InErrorCtxt, *FString::Printf(TEXT("Failed to convert return property '%s' (%s) when calling %s"), *InOutputParams[ReturnPropIndex].ParamProp->GetName(), *InOutputParams[ReturnPropIndex].ParamProp->GetClass()->GetName(), InCallingCtxt));
+				return false;
+			}
+		}
+	}
+
+	return true;
+}
+
+FString BuildFunctionDocString(const UFunction* InFunc, const FString& InFuncPythonName, const TArray<FGeneratedWrappedMethodParameter>& InInputParams, const TArray<FGeneratedWrappedMethodParameter>& InOutputParams, const bool* InStaticOverride)
+{
+	const bool bIsStatic = (InStaticOverride) ? *InStaticOverride : InFunc->HasAnyFunctionFlags(FUNC_Static);
+
+	FString FunctionDeclDocString = FString::Printf(TEXT("%s.%s("), (bIsStatic ? TEXT("X") : TEXT("x")), *InFuncPythonName);
+	for (const FGeneratedWrappedMethodParameter& InputParam : InInputParams)
+	{
+		if (FunctionDeclDocString[FunctionDeclDocString.Len() - 1] != TEXT('('))
+		{
+			FunctionDeclDocString += TEXT(", ");
+		}
+		FunctionDeclDocString += UTF8_TO_TCHAR(InputParam.ParamName.GetData());
+		if (InputParam.ParamDefaultValue.IsSet())
+		{
+			FunctionDeclDocString += TEXT('=');
+			FunctionDeclDocString += InputParam.ParamDefaultValue.GetValue();
+		}
+	}
+	FunctionDeclDocString += TEXT(")");
+
+	if (InOutputParams.Num() > 0)
+	{
+		FunctionDeclDocString += TEXT(" -> ");
+
+		// If we have multiple return values and the main return value is a bool, we return None (for false) or the (potentially packed) return value without the bool (for true)
+		int32 IndexOffset = 0;
+		if (InOutputParams.Num() > 1)
+		{
+			if (InOutputParams[0].ParamProp->HasAnyPropertyFlags(CPF_ReturnParm) && InOutputParams[0].ParamProp->IsA<UBoolProperty>())
+			{
+				++IndexOffset;
+			}
+		}
+
+		if (InOutputParams.Num() - IndexOffset == 1)
+		{
+			FunctionDeclDocString += GetPropertyTypePythonName(InOutputParams[IndexOffset].ParamProp);
+		}
+		else
+		{
+			const bool bHasReturnValue = InOutputParams[0].ParamProp->HasAnyPropertyFlags(CPF_ReturnParm);
+			FunctionDeclDocString += TEXT('(');
+			for (int32 OutParamIndex = IndexOffset; OutParamIndex < InOutputParams.Num(); ++OutParamIndex)
+			{
+				if (OutParamIndex > IndexOffset)
+				{
+					FunctionDeclDocString += TEXT(", ");
+				}
+				if (OutParamIndex > 0 || !bHasReturnValue)
+				{
+					FunctionDeclDocString += UTF8_TO_TCHAR(InOutputParams[OutParamIndex].ParamName.GetData());
+					FunctionDeclDocString += TEXT('=');
+				}
+				FunctionDeclDocString += GetPropertyTypePythonName(InOutputParams[OutParamIndex].ParamProp);
+			}
+			FunctionDeclDocString += TEXT(')');
+		}
+
+		if (IndexOffset > 0)
+		{
+			FunctionDeclDocString += TEXT(" or None");
+		}
+	}
+
+	return FunctionDeclDocString;
+}
+
 bool IsBlueprintExposedClass(const UClass* InClass)
 {
 	for (const UClass* ParentClass = InClass; ParentClass; ParentClass = ParentClass->GetSuperClass())
@@ -342,7 +804,10 @@ bool IsBlueprintExposedProperty(const UProperty* InProp)
 
 bool IsBlueprintExposedFunction(const UFunction* InFunc)
 {
-	return InFunc->HasAnyFunctionFlags(FUNC_BlueprintCallable | FUNC_BlueprintEvent) && !InFunc->HasMetaData(BlueprintGetterMetaDataKey) && !InFunc->HasMetaData(BlueprintSetterMetaDataKey);
+	return InFunc->HasAnyFunctionFlags(FUNC_BlueprintCallable | FUNC_BlueprintEvent)
+		&& !InFunc->HasMetaData(BlueprintGetterMetaDataKey)
+		&& !InFunc->HasMetaData(BlueprintSetterMetaDataKey)
+		&& !InFunc->HasMetaData(CustomStructureParamMetaDataKey);
 }
 
 bool IsBlueprintExposedField(const UField* InField)
@@ -390,6 +855,63 @@ bool IsBlueprintGeneratedEnum(const UEnum* InEnum)
 	return InEnum->IsA<UUserDefinedEnum>();
 }
 
+bool IsDeprecatedClass(const UClass* InClass, FString* OutDeprecationMessage)
+{
+	if (InClass->HasAnyClassFlags(CLASS_Deprecated))
+	{
+		if (OutDeprecationMessage)
+		{
+			*OutDeprecationMessage = InClass->GetMetaData(DeprecationMessageMetaDataKey);
+			if (OutDeprecationMessage->IsEmpty())
+			{
+				*OutDeprecationMessage = FString::Printf(TEXT("Class '%s' is deprecated."), *InClass->GetName());
+			}
+		}
+
+		return true;
+	}
+
+	return false;
+}
+
+bool IsDeprecatedProperty(const UProperty* InProp, FString* OutDeprecationMessage)
+{
+	if (InProp->HasMetaData(DeprecatedPropertyMetaDataKey))
+	{
+		if (OutDeprecationMessage)
+		{
+			*OutDeprecationMessage = InProp->GetMetaData(DeprecationMessageMetaDataKey);
+			if (OutDeprecationMessage->IsEmpty())
+			{
+				*OutDeprecationMessage = FString::Printf(TEXT("Property '%s' is deprecated."), *InProp->GetName());
+			}
+		}
+
+		return true;
+	}
+
+	return false;
+}
+
+bool IsDeprecatedFunction(const UFunction* InFunc, FString* OutDeprecationMessage)
+{
+	if (InFunc->HasMetaData(DeprecatedFunctionMetaDataKey))
+	{
+		if (OutDeprecationMessage)
+		{
+			*OutDeprecationMessage = InFunc->GetMetaData(DeprecationMessageMetaDataKey);
+			if (OutDeprecationMessage->IsEmpty())
+			{
+				*OutDeprecationMessage = FString::Printf(TEXT("Function '%s' is deprecated."), *InFunc->GetName());
+			}
+		}
+
+		return true;
+	}
+
+	return false;
+}
+
 bool ShouldExportClass(const UClass* InClass)
 {
 	return IsBlueprintExposedClass(InClass) || HasBlueprintExposedFields(InClass);
@@ -413,13 +935,13 @@ bool ShouldExportEnumEntry(const UEnum* InEnum, int32 InEnumEntryIndex)
 bool ShouldExportProperty(const UProperty* InProp)
 {
 	const bool bCanScriptExport = !InProp->HasMetaData(ScriptNoExportMetaDataKey);
-	return bCanScriptExport && IsBlueprintExposedProperty(InProp);
+	return bCanScriptExport && (IsBlueprintExposedProperty(InProp) || IsDeprecatedProperty(InProp));
 }
 
 bool ShouldExportEditorOnlyProperty(const UProperty* InProp)
 {
 	const bool bCanScriptExport = !InProp->HasMetaData(ScriptNoExportMetaDataKey);
-	return bCanScriptExport && GIsEditor && InProp->HasAnyPropertyFlags(CPF_Edit);
+	return bCanScriptExport && GIsEditor && (InProp->HasAnyPropertyFlags(CPF_Edit) || IsDeprecatedProperty(InProp));
 }
 
 bool ShouldExportFunction(const UFunction* InFunc)
@@ -430,6 +952,40 @@ bool ShouldExportFunction(const UFunction* InFunc)
 
 FString PythonizeName(const FString& InName, const EPythonizeNameCase InNameCase)
 {
+	static const TSet<FString, FCaseSensitiveStringSetFuncs> ReservedKeywords = {
+		TEXT("and"),
+		TEXT("as"),
+		TEXT("assert"),
+		TEXT("break"),
+		TEXT("class"),
+		TEXT("continue"),
+		TEXT("def"),
+		TEXT("del"),
+		TEXT("elif"),
+		TEXT("else"),
+		TEXT("except"),
+		TEXT("finally"),
+		TEXT("for"),
+		TEXT("from"),
+		TEXT("global"),
+		TEXT("if"),
+		TEXT("import"),
+		TEXT("in"),
+		TEXT("is"),
+		TEXT("lambda"),
+		TEXT("nonlocal"),
+		TEXT("not"),
+		TEXT("or"),
+		TEXT("pass"),
+		TEXT("raise"),
+		TEXT("return"),
+		TEXT("try"),
+		TEXT("while"),
+		TEXT("with"),
+		TEXT("yield"),
+		TEXT("property"),
+	};
+
 	FString PythonizedName;
 	PythonizedName.Reserve(InName.Len() + 10);
 
@@ -476,6 +1032,12 @@ FString PythonizeName(const FString& InName, const EPythonizeNameCase InNameCase
 		PythonizedName.ToUpperInline();
 	}
 
+	// Don't allow the name to conflict with a keyword
+	if (ReservedKeywords.Contains(PythonizedName))
+	{
+		PythonizedName += TEXT('_');
+	}
+
 	return PythonizedName;
 }
 
@@ -500,11 +1062,11 @@ FString PythonizePropertyName(const FString& InName, const EPythonizeNameCase In
 		}
 
 		// Strip the "Out" prefix from names
-		if (InName.Len() - NameOffset >= 4 && InName[NameOffset] == TEXT('O') && InName[NameOffset + 1] == TEXT('u') && InName[NameOffset + 2] == TEXT('t') && FChar::IsUpper(InName[NameOffset + 3]))
-		{
-			NameOffset += 3;
-			continue;
-		}
+		//if (InName.Len() - NameOffset >= 4 && InName[NameOffset] == TEXT('O') && InName[NameOffset + 1] == TEXT('u') && InName[NameOffset + 2] == TEXT('t') && FChar::IsUpper(InName[NameOffset + 3]))
+		//{
+		//	NameOffset += 3;
+		//	continue;
+		//}
 
 		// Nothing more to strip
 		break;
@@ -518,9 +1080,11 @@ FString PythonizePropertyTooltip(const FString& InTooltip, const UProperty* InPr
 	return PythonizeTooltip(InTooltip, FPythonizeTooltipContext(InProp, nullptr, InReadOnlyFlags));
 }
 
-FString PythonizeFunctionTooltip(const FString& InTooltip, const UStruct* InParamsStruct)
+FString PythonizeFunctionTooltip(const FString& InTooltip, const UFunction* InFunc, const TSet<FName>& ParamsToIgnore)
 {
-	return PythonizeTooltip(InTooltip, FPythonizeTooltipContext(nullptr, InParamsStruct));
+	FPythonizeTooltipContext PythonizeTooltipContext(nullptr, InFunc);
+	PythonizeTooltipContext.ParamsToIgnore = ParamsToIgnore;
+	return PythonizeTooltip(InTooltip, PythonizeTooltipContext);
 }
 
 FString PythonizeTooltip(const FString& InTooltip, const FPythonizeTooltipContext& InContext)
@@ -618,7 +1182,11 @@ FString PythonizeTooltip(const FString& InTooltip, const FPythonizeTooltipContex
 				SkipToNextToken();
 				ParseComplexToken(ParamComment);
 
-				ParsedParamTokens.Add(MakeTuple(FName(*ParamName), MoveTemp(ParamComment)));
+				const FName ParamFName = *ParamName;
+				if (!InContext.ParamsToIgnore.Contains(ParamFName))
+				{
+					ParsedParamTokens.Add(MakeTuple(ParamFName, MoveTemp(ParamComment)));
+				}
 			}
 			else if (TokenName == TEXT("return") || TokenName == TEXT("returns"))
 			{
@@ -657,6 +1225,13 @@ FString PythonizeTooltip(const FString& InTooltip, const FPythonizeTooltipContex
 
 	PythonizedTooltip.TrimEndInline();
 
+	// Add the deprecation message
+	if (!InContext.DeprecationMessage.IsEmpty())
+	{
+		PythonizedTooltip += TEXT("\ndeprecated: ");
+		PythonizedTooltip += InContext.DeprecationMessage;
+	}
+
 	// Process the misc tokens into PythonizedTooltip
 	for (const auto& MiscTokenPair : ParsedMiscTokens)
 	{
@@ -679,9 +1254,9 @@ FString PythonizeTooltip(const FString& InTooltip, const FPythonizeTooltipContex
 		PythonizedTooltip += TEXT("param: ");
 		PythonizedTooltip += PythonizePropertyName(ParamTokenPair.Key.ToString(), EPythonizeNameCase::Lower);
 
-		if (InContext.ParamsStruct)
+		if (InContext.Func)
 		{
-			if (const UProperty* ParamProp = InContext.ParamsStruct->FindPropertyByName(ParamTokenPair.Key))
+			if (const UProperty* ParamProp = InContext.Func->FindPropertyByName(ParamTokenPair.Key))
 			{
 				AppendParamTypeDoc(ParamProp);
 			}
@@ -693,12 +1268,17 @@ FString PythonizeTooltip(const FString& InTooltip, const FPythonizeTooltipContex
 			PythonizedTooltip += ParamTokenPair.Value;
 		}
 	}
-	if (InContext.ParamsStruct)
+	if (InContext.Func)
 	{
-		for (TFieldIterator<const UProperty> ParamIt(InContext.ParamsStruct); ParamIt; ++ParamIt)
+		for (TFieldIterator<const UProperty> ParamIt(InContext.Func); ParamIt; ++ParamIt)
 		{
 			const UProperty* ParamProp = *ParamIt;
 
+			if (InContext.ParamsToIgnore.Contains(ParamProp->GetFName()))
+			{
+				continue;
+			}
+
 			const bool bHasProcessedParamProp = ParsedParamTokens.ContainsByPredicate([&ParamProp](const TTuple<FName, FString>& ParamTokenPair)
 			{
 				return ParamTokenPair.Key == ParamProp->GetFName();
