Copying //UE4/Dev-Enterprise to //UE4/Dev-Main (Source: //UE4/Dev-Enterprise @ 3972172)

diff --git a/Engine/Plugins/Experimental/PythonScriptPlugin/Source/PythonScriptPlugin/Private/PyGenUtil.h b/Engine/Plugins/Experimental/PythonScriptPlugin/Source/PythonScriptPlugin/Private/PyGenUtil.h
--- a/Engine/Plugins/Experimental/PythonScriptPlugin/Source/PythonScriptPlugin/Private/PyGenUtil.h
+++ b/Engine/Plugins/Experimental/PythonScriptPlugin/Source/PythonScriptPlugin/Private/PyGenUtil.h
@@ -16,11 +16,17 @@ namespace PyGenUtil
 {
 	extern const FName ScriptNameMetaDataKey;
 	extern const FName ScriptNoExportMetaDataKey;
+	extern const FName ScriptMethodMetaDataKey;
+	extern const FName ScriptMathOpMetaDataKey;
 	extern const FName BlueprintTypeMetaDataKey;
 	extern const FName NotBlueprintTypeMetaDataKey;
 	extern const FName BlueprintSpawnableComponentMetaDataKey;
 	extern const FName BlueprintGetterMetaDataKey;
 	extern const FName BlueprintSetterMetaDataKey;
+	extern const FName CustomStructureParamMetaDataKey;
+	extern const FName DeprecatedPropertyMetaDataKey;
+	extern const FName DeprecatedFunctionMetaDataKey;
+	extern const FName DeprecationMessageMetaDataKey;
 
 	/** Name used by the Python equivalent of PostInitProperties */
 	static const char* PostInitFuncName = "_post_init";
@@ -31,81 +37,225 @@ namespace PyGenUtil
 	/** Stores the data needed by a runtime generated Python parameter */
 	struct FGeneratedWrappedMethodParameter
 	{
+		FGeneratedWrappedMethodParameter()
+			: ParamProp(nullptr)
+		{
+		}
+
+		FGeneratedWrappedMethodParameter(FGeneratedWrappedMethodParameter&&) = default;
+		FGeneratedWrappedMethodParameter(const FGeneratedWrappedMethodParameter&) = default;
+		FGeneratedWrappedMethodParameter& operator=(FGeneratedWrappedMethodParameter&&) = default;
+		FGeneratedWrappedMethodParameter& operator=(const FGeneratedWrappedMethodParameter&) = default;
+
 		/** The name of the parameter */
 		FUTF8Buffer ParamName;
 
-		/** The Unreal property name for this parameter */
-		FName ParamPropName;
+		/** The Unreal property for this parameter */
+		const UProperty* ParamProp;
 
 		/** The default Unreal ExportText value of this parameter; parameters with this set are considered optional */
 		TOptional<FString> ParamDefaultValue;
 	};
 
+	/** Stores the data needed to call an Unreal function via Python */
+	struct FGeneratedWrappedFunction
+	{
+		FGeneratedWrappedFunction()
+			: Func(nullptr)
+		{
+		}
+
+		FGeneratedWrappedFunction(FGeneratedWrappedFunction&&) = default;
+		FGeneratedWrappedFunction(const FGeneratedWrappedFunction&) = default;
+		FGeneratedWrappedFunction& operator=(FGeneratedWrappedFunction&&) = default;
+		FGeneratedWrappedFunction& operator=(const FGeneratedWrappedFunction&) = default;
+
+		/** Set the function to call, and also extract the parameter lists */
+		void SetFunctionAndExtractParams(const UFunction* InFunc);
+
+		/** The Unreal function to call (static dispatch) */
+		const UFunction* Func;
+
+		/** Array of input parameters associated with the function */
+		TArray<FGeneratedWrappedMethodParameter> InputParams;
+
+		/** Array of output (including return) parameters associated with the function */
+		TArray<FGeneratedWrappedMethodParameter> OutputParams;
+	};
+
 	/** Stores the data needed by a runtime generated Python method */
 	struct FGeneratedWrappedMethod
 	{
 		FGeneratedWrappedMethod()
-			: Class(nullptr)
-			, MethodCallback(nullptr)
+			: MethodCallback(nullptr)
 			, MethodFlags(0)
 		{
 		}
 
+		FGeneratedWrappedMethod(FGeneratedWrappedMethod&&) = default;
+		FGeneratedWrappedMethod(const FGeneratedWrappedMethod&) = default;
+		FGeneratedWrappedMethod& operator=(FGeneratedWrappedMethod&&) = default;
+		FGeneratedWrappedMethod& operator=(const FGeneratedWrappedMethod&) = default;
+
 		/** Convert this wrapper type to its Python type */
 		void ToPython(FPyMethodWithClosureDef& OutPyMethod) const;
 
-		/** The class to call the method on (static dispatch) */
-		UClass* Class;
-
 		/** The name of the method */
 		FUTF8Buffer MethodName;
 
 		/** The doc string of the method */
 		FUTF8Buffer MethodDoc;
 
-		/** The Unreal function name for this method */
-		FName MethodFuncName;
+		/** The Unreal function for this method */
+		FGeneratedWrappedFunction MethodFunc;
 
 		/* The C function this method should call */
 		PyCFunctionWithClosure MethodCallback;
 
 		/* The METH_ flags for this method */
 		int MethodFlags;
+	};
+
+	/** Stores the data needed for a set of runtime generated Python methods */
+	struct FGeneratedWrappedMethods
+	{
+		FGeneratedWrappedMethods() = default;
+		FGeneratedWrappedMethods(FGeneratedWrappedMethods&&) = default;
+		FGeneratedWrappedMethods(const FGeneratedWrappedMethods&) = delete;
+		FGeneratedWrappedMethods& operator=(FGeneratedWrappedMethods&&) = default;
+		FGeneratedWrappedMethods& operator=(const FGeneratedWrappedMethods&) = delete;
+
+		/** Called to ready the generated methods with Python */
+		void Finalize();
 
-		/** Array of parameters associated with this method */
-		TArray<FGeneratedWrappedMethodParameter> MethodParams;
+		/** Array of methods associated from the wrapped type */
+		TArray<FGeneratedWrappedMethod> TypeMethods;
+
+		/** The Python methods that were generated from TypeMethods (in Finalize) */
+		TArray<FPyMethodWithClosureDef> PyMethods;
+	};
+
+	/** Stores the data needed by a runtime generated Python method that is dynamically created and registered post-finalize of its owner struct type (for hoisting util functions onto structs) */
+	struct FGeneratedWrappedDynamicStructMethod : public FGeneratedWrappedMethod
+	{
+		FGeneratedWrappedDynamicStructMethod() = default;
+		FGeneratedWrappedDynamicStructMethod(FGeneratedWrappedDynamicStructMethod&&) = default;
+		FGeneratedWrappedDynamicStructMethod(const FGeneratedWrappedDynamicStructMethod&) = default;
+		FGeneratedWrappedDynamicStructMethod& operator=(FGeneratedWrappedDynamicStructMethod&&) = default;
+		FGeneratedWrappedDynamicStructMethod& operator=(const FGeneratedWrappedDynamicStructMethod&) = default;
+
+		/** The struct parameter information (this parameter is set to the struct instance that calls the method) */
+		FGeneratedWrappedMethodParameter StructParam;
+	};
+
+	/** Stores the data needed by a runtime generated Python method that is dynamically created and registered post-finalize of its owner struct type (for hoisting util functions onto structs) */
+	struct FGeneratedWrappedDynamicStructMethodWithClosure : public FGeneratedWrappedDynamicStructMethod
+	{
+		FGeneratedWrappedDynamicStructMethodWithClosure() = default;
+		FGeneratedWrappedDynamicStructMethodWithClosure(FGeneratedWrappedDynamicStructMethodWithClosure&&) = default;
+		FGeneratedWrappedDynamicStructMethodWithClosure(const FGeneratedWrappedDynamicStructMethodWithClosure&) = delete;
+		FGeneratedWrappedDynamicStructMethodWithClosure& operator=(FGeneratedWrappedDynamicStructMethodWithClosure&&) = default;
+		FGeneratedWrappedDynamicStructMethodWithClosure& operator=(const FGeneratedWrappedDynamicStructMethodWithClosure&) = delete;
+
+		/** Called to ready the generated method for Python */
+		void Finalize();
+
+		/** Python method that was generated from this method */
+		FPyMethodWithClosureDef PyMethod;
+	};
+
+	/** Stores the data needed by a runtime generated Python operator stack function */
+	struct FGeneratedWrappedStructMathOpFunction : public FGeneratedWrappedFunction
+	{
+		FGeneratedWrappedStructMathOpFunction() = default;
+		FGeneratedWrappedStructMathOpFunction(FGeneratedWrappedStructMathOpFunction&&) = default;
+		FGeneratedWrappedStructMathOpFunction(const FGeneratedWrappedStructMathOpFunction&) = default;
+		FGeneratedWrappedStructMathOpFunction& operator=(FGeneratedWrappedStructMathOpFunction&&) = default;
+		FGeneratedWrappedStructMathOpFunction& operator=(const FGeneratedWrappedStructMathOpFunction&) = default;
+
+		/** Set the function to call, and also extract the parameter lists */
+		bool SetFunctionAndExtractParams(const UFunction* InFunc);
+
+		/** The struct parameter information (this parameter is set to the struct instance that calls the function) */
+		FGeneratedWrappedMethodParameter StructParam;
+	};
+
+	/** Stores the data needed by a runtime generated Python operator stack that is dynamically created and registered post-finalize of its owner struct type (for hoisting math operators onto structs) */
+	struct FGeneratedWrappedStructMathOpStack
+	{
+		/** Known math operator types */
+		enum class EOpType : uint8
+		{
+			Add = 0,			// +
+			InlineAdd,			// +=
+			Subtract,			// -
+			InlineSubtract,		// -=
+			Multiply,			// *
+			InlineMultiply,		// *=
+			Divide,				// /
+			InlineDivide,		// /=
+			Modulus,			// %
+			InlineModulus,		// %=
+			And,				// &
+			InlineAnd,			// &=
+			Or,					// |
+			InlineOr,			// |=
+			Xor,				// ^
+			InlineXor,			// ^=
+			RightShift,			// >>
+			InlineRightShift,	// >>=
+			LeftShift,			// <<
+			InlineLeftShift,	// <<=
+			Num,
+		};
+
+		/**
+		 * Given a potential operator string, try and convert it to a known operator type
+		 * @return true if the conversion was a success, false otherwise
+		 */
+		static bool StringToOpType(const TCHAR* InStr, EOpType& OutOpType);
+
+		/**
+		 * Is the given operator a inline operator?
+		 */
+		static bool IsInlineOp(const EOpType InOpType);
+
+		/** Array of math operator functions associated with this operator stack */
+		TArray<FGeneratedWrappedStructMathOpFunction> MathOpFuncs;
 	};
 
 	/** Stores the data needed by a runtime generated Python get/set */
 	struct FGeneratedWrappedGetSet
 	{
 		FGeneratedWrappedGetSet()
-			: Class(nullptr)
+			: Prop(nullptr)
 			, GetCallback(nullptr)
 			, SetCallback(nullptr)
 		{
 		}
 
+		FGeneratedWrappedGetSet(FGeneratedWrappedGetSet&&) = default;
+		FGeneratedWrappedGetSet(const FGeneratedWrappedGetSet&) = default;
+		FGeneratedWrappedGetSet& operator=(FGeneratedWrappedGetSet&&) = default;
+		FGeneratedWrappedGetSet& operator=(const FGeneratedWrappedGetSet&) = default;
+
 		/** Convert this wrapper type to its Python type */
 		void ToPython(PyGetSetDef& OutPyGetSet) const;
 
-		/** The class to call the get/set functions on (static dispatch) */
-		UClass* Class;
-
 		/** The name of the get/set */
 		FUTF8Buffer GetSetName;
 
 		/** The doc string of the get/set */
 		FUTF8Buffer GetSetDoc;
 
-		/** The Unreal property name for this get/set */
-		FName PropName;
+		/** The Unreal property for this get/set */
+		const UProperty* Prop;
 
 		/** The Unreal function for the get (if any) */
-		FName GetFuncName;
+		FGeneratedWrappedFunction GetFunc;
 
 		/** The Unreal function for the set (if any) */
-		FName SetFuncName;
+		FGeneratedWrappedFunction SetFunc;
 
 		/* The C function that should be called for get */
 		getter GetCallback;
@@ -114,11 +264,35 @@ namespace PyGenUtil
 		setter SetCallback;
 	};
 
+	/** Stores the data needed for a set of runtime generated Python get/sets */
+	struct FGeneratedWrappedGetSets
+	{
+		FGeneratedWrappedGetSets() = default;
+		FGeneratedWrappedGetSets(FGeneratedWrappedGetSets&&) = default;
+		FGeneratedWrappedGetSets(const FGeneratedWrappedGetSets&) = delete;
+		FGeneratedWrappedGetSets& operator=(FGeneratedWrappedGetSets&&) = default;
+		FGeneratedWrappedGetSets& operator=(const FGeneratedWrappedGetSets&) = delete;
+
+		/** Called to ready the generated get/sets with Python */
+		void Finalize();
+
+		/** Array of get/sets from the wrapped type */
+		TArray<FGeneratedWrappedGetSet> TypeGetSets;
+
+		/** The Python get/sets that were generated from TypeGetSets (in Finalize) */
+		TArray<PyGetSetDef> PyGetSets;
+	};
+
 	/** Stores the data needed to generate a Python doc string for editor exposed properties */
 	struct FGeneratedWrappedPropertyDoc
 	{
 		explicit FGeneratedWrappedPropertyDoc(const UProperty* InProp);
 
+		FGeneratedWrappedPropertyDoc(FGeneratedWrappedPropertyDoc&&) = default;
+		FGeneratedWrappedPropertyDoc(const FGeneratedWrappedPropertyDoc&) = default;
+		FGeneratedWrappedPropertyDoc& operator=(FGeneratedWrappedPropertyDoc&&) = default;
+		FGeneratedWrappedPropertyDoc& operator=(const FGeneratedWrappedPropertyDoc&) = default;
+
 		/** Util function to sort an array of doc structs based on the Pythonized property name */
 		static bool SortPredicate(const FGeneratedWrappedPropertyDoc& InOne, const FGeneratedWrappedPropertyDoc& InTwo);
 
@@ -138,7 +312,7 @@ namespace PyGenUtil
 		FString EditorDocString;
 	};
 
-	/** Stores the data needed by a runtime generated Python type */
+	/** Stores the minimal data needed by a runtime generated Python type */
 	struct FGeneratedWrappedType
 	{
 		FGeneratedWrappedType()
@@ -146,40 +320,97 @@ namespace PyGenUtil
 			PyType = { PyVarObject_HEAD_INIT(nullptr, 0) };
 		}
 
+		virtual ~FGeneratedWrappedType() = default;
+
+		FGeneratedWrappedType(FGeneratedWrappedType&&) = default;
+		FGeneratedWrappedType(const FGeneratedWrappedType&) = delete;
+		FGeneratedWrappedType& operator=(FGeneratedWrappedType&&) = default;
+		FGeneratedWrappedType& operator=(const FGeneratedWrappedType&) = delete;
+
 		/** Called to ready the generated type with Python */
 		bool Finalize();
 
+		/** Internal version of Finalize, called before readying the type with Python */
+		virtual void Finalize_PreReady();
+
+		/** Internal version of Finalize, called after readying the type with Python */
+		virtual void Finalize_PostReady();
+
 		/** The name of the type */
 		FUTF8Buffer TypeName;
 
 		/** The doc string of the type */
 		FUTF8Buffer TypeDoc;
 
-		/** Array of methods associated with this type */
-		TArray<FGeneratedWrappedMethod> TypeMethods;
+		/** The meta-data associated with this type */
+		TSharedPtr<FPyWrapperBaseMetaData> MetaData;
 
-		/** The Python methods that were generated from TypeMethods (in Finalize) */
-		TArray<FPyMethodWithClosureDef> PyMethods;
+		/* The Python type that was generated */
+		PyTypeObject PyType;
+	};
 
-		/** Array of get/sets associated with this type */
-		TArray<FGeneratedWrappedGetSet> TypeGetSets;
+	/** Stores the data needed by a runtime generated Python struct type */
+	struct FGeneratedWrappedStructType : public FGeneratedWrappedType
+	{
+		FGeneratedWrappedStructType() = default;
+		FGeneratedWrappedStructType(FGeneratedWrappedStructType&&) = default;
+		FGeneratedWrappedStructType(const FGeneratedWrappedStructType&) = delete;
+		FGeneratedWrappedStructType& operator=(FGeneratedWrappedStructType&&) = default;
+		FGeneratedWrappedStructType& operator=(const FGeneratedWrappedStructType&) = delete;
 
-		/** The Python get/sets that were generated from TypeGetSets (in Finalize) */
-		TArray<PyGetSetDef> PyGetSets;
+		/** Internal version of Finalize, called before readying the type with Python */
+		virtual void Finalize_PreReady() override;
+
+		/** Called to add a dynamic struct method to this Python type (this should only be called post-finalize) */
+		void AddDynamicStructMethod(FGeneratedWrappedDynamicStructMethod&& InDynamicStructMethod);
+
+		/** Get/sets associated with this type */
+		FGeneratedWrappedGetSets GetSets;
 
 		/** The doc string data for the properties associated with this type */
-		TArray<FGeneratedWrappedPropertyDoc> TypePropertyDocs;
+		TArray<FGeneratedWrappedPropertyDoc> PropertyDocs;
 
-		/** The meta-data associated with this type */
-		TSharedPtr<FPyWrapperBaseMetaData> MetaData;
+		/** Array of dynamic struct methods associated with this struct type (call AddDynamicStructMethod to add methods) */
+		TArray<TSharedRef<FGeneratedWrappedDynamicStructMethodWithClosure>> DynamicStructMethods;
 
-		/* The Python type that was generated */
-		PyTypeObject PyType;
+		/** Python number methods for this struct */
+		PyNumberMethods PyNumber;
+	};
+
+	/** Stores the data needed by a runtime generated Python class type */
+	struct FGeneratedWrappedClassType : public FGeneratedWrappedType
+	{
+		FGeneratedWrappedClassType() = default;
+		FGeneratedWrappedClassType(FGeneratedWrappedClassType&&) = default;
+		FGeneratedWrappedClassType(const FGeneratedWrappedClassType&) = delete;
+		FGeneratedWrappedClassType& operator=(FGeneratedWrappedClassType&&) = default;
+		FGeneratedWrappedClassType& operator=(const FGeneratedWrappedClassType&) = delete;
+
+		/** Internal version of Finalize, called before readying the type with Python */
+		virtual void Finalize_PreReady() override;
+
+		/** Internal version of Finalize, called after readying the type with Python */
+		virtual void Finalize_PostReady() override;
+
+		/** Methods associated with this type */
+		FGeneratedWrappedMethods Methods;
+
+		/** Get/sets associated with this type */
+		FGeneratedWrappedGetSets GetSets;
+
+		/** The doc string data for the properties associated with this type */
+		TArray<FGeneratedWrappedPropertyDoc> PropertyDocs;
 	};
 
 	/** Definition data for an Unreal property generated from a Python type */
 	struct FPropertyDef
 	{
+		FPropertyDef() = default;
+		FPropertyDef(FPropertyDef&&) = default;
+		FPropertyDef(const FPropertyDef&) = delete;
+		FPropertyDef& operator=(FPropertyDef&&) = default;
+		FPropertyDef& operator=(const FPropertyDef&) = delete;
+
 		/** Data needed to re-wrap this property for Python */
 		FGeneratedWrappedGetSet GeneratedWrappedGetSet;
 
@@ -190,6 +421,12 @@ namespace PyGenUtil
 	/** Definition data for an Unreal function generated from a Python type */
 	struct FFunctionDef
 	{
+		FFunctionDef() = default;
+		FFunctionDef(FFunctionDef&&) = default;
+		FFunctionDef(const FFunctionDef&) = delete;
+		FFunctionDef& operator=(FFunctionDef&&) = default;
+		FFunctionDef& operator=(const FFunctionDef&) = delete;
+
 		/** Data needed to re-wrap this function for Python */
 		FGeneratedWrappedMethod GeneratedWrappedMethod;
 
@@ -217,23 +454,24 @@ namespace PyGenUtil
 	{
 		FPythonizeTooltipContext()
 			: Prop(nullptr)
-			, ParamsStruct(nullptr)
+			, Func(nullptr)
 			, ReadOnlyFlags(CPF_BlueprintReadOnly | CPF_EditConst)
 		{
 		}
 
-		FPythonizeTooltipContext(const UProperty* InProp, const UStruct* InParamsStruct, const uint64 InReadOnlyFlags = CPF_BlueprintReadOnly | CPF_EditConst)
-			: Prop(InProp)
-			, ParamsStruct(InParamsStruct)
-			, ReadOnlyFlags(InReadOnlyFlags)
-		{
-		}
+		FPythonizeTooltipContext(const UProperty* InProp, const UFunction* InFunc, const uint64 InReadOnlyFlags = CPF_BlueprintReadOnly | CPF_EditConst);
 
 		/** Optional property that should be used when converting property tooltips */
 		const UProperty* Prop;
 
-		/** Optional structure of params that should be used when converting function tooltips */
-		const UStruct* ParamsStruct;
+		/** Optional function that should be used when converting function tooltips */
+		const UFunction* Func;
+
+		/** Optional deprecation message for the property or function */
+		FString DeprecationMessage;
+
+		/** Optional set of parameters that we should ignore when generating function tooltips */
+		TSet<FName> ParamsToIgnore;
 
 		/** Flags that dictate whether the property should be considered read-only */
 		uint64 ReadOnlyFlags;
@@ -246,11 +484,26 @@ namespace PyGenUtil
 	PyObject* GetPostInitFunc(PyTypeObject* InPyType);
 
 	/** Add a struct init parameter to the given array of method parameters */
-	void AddStructInitParam(const TCHAR* InUnrealPropName, const TCHAR* InPythonAttrName, TArray<FGeneratedWrappedMethodParameter>& OutInitParams);
+	void AddStructInitParam(const UProperty* InUnrealProp, const TCHAR* InPythonAttrName, TArray<FGeneratedWrappedMethodParameter>& OutInitParams);
+
+	/** Given a function, extract all of its parameter information (input and output) */
+	void ExtractFunctionParams(const UFunction* InFunc, TArray<FGeneratedWrappedMethodParameter>& OutInputParams, TArray<FGeneratedWrappedMethodParameter>& OutOutputParams);
+
+	/** Apply default values to function arguments */
+	void ApplyParamDefaults(void* InBaseParamsAddr, const TArray<FGeneratedWrappedMethodParameter>& InParamDef);
 
 	/** Parse an array Python objects from the args and keywords based on the generated method parameter data */
 	bool ParseMethodParameters(PyObject* InArgs, PyObject* InKwds, const TArray<FGeneratedWrappedMethodParameter>& InParamDef, const char* InPyMethodName, TArray<PyObject*>& OutPyParams);
 
+	/** Given a set of return values and the struct data associated with them, pack them appropriately for returning to Python */
+	PyObject* PackReturnValues(void* InBaseParamsAddr, const TArray<FGeneratedWrappedMethodParameter>& InOutputParams, const TCHAR* InErrorCtxt, const TCHAR* InCallingCtxt);
+
+	/** Given a Python return value, unpack the values into the struct data associated with them */
+	bool UnpackReturnValues(PyObject* InRetVals, void* InBaseParamsAddr, const TArray<FGeneratedWrappedMethodParameter>& InOutputParams, const TCHAR* InErrorCtxt, const TCHAR* InCallingCtxt);
+
+	/** Build a Python doc string for the given function and arguments list */
+	FString BuildFunctionDocString(const UFunction* InFunc, const FString& InFuncPythonName, const TArray<FGeneratedWrappedMethodParameter>& InInputParams, const TArray<FGeneratedWrappedMethodParameter>& InOutputParams, const bool* InStaticOverride = nullptr);
+
 	/** Is the given class generated by Blueprints? */
 	bool IsBlueprintGeneratedClass(const UClass* InClass);
 
@@ -260,6 +513,15 @@ namespace PyGenUtil
 	/** Is the given enum generated by Blueprints? */
 	bool IsBlueprintGeneratedEnum(const UEnum* InEnum);
 
+	/** Is the given class marked as deprecated? */
+	bool IsDeprecatedClass(const UClass* InClass, FString* OutDeprecationMessage = nullptr);
+
+	/** Is the given property marked as deprecated? */
+	bool IsDeprecatedProperty(const UProperty* InProp, FString* OutDeprecationMessage = nullptr);
+
+	/** Is the given function marked as deprecated? */
+	bool IsDeprecatedFunction(const UFunction* InFunc, FString* OutDeprecationMessage = nullptr);
+
 	/** Should the given class be exported to Python? */
 	bool ShouldExportClass(const UClass* InClass);
 
@@ -291,7 +553,7 @@ namespace PyGenUtil
 	FString PythonizePropertyTooltip(const FString& InTooltip, const UProperty* InProp, const uint64 InReadOnlyFlags = CPF_BlueprintReadOnly | CPF_EditConst);
 
 	/** Given a function tooltip, convert it to a doc string */
-	FString PythonizeFunctionTooltip(const FString& InTooltip, const UStruct* InParamsStruct);
+	FString PythonizeFunctionTooltip(const FString& InTooltip, const UFunction* InFunc, const TSet<FName>& ParamsToIgnore = TSet<FName>());
 
 	/** Given a tooltip, convert it to a doc string */
 	FString PythonizeTooltip(const FString& InTooltip, const FPythonizeTooltipContext& InContext = FPythonizeTooltipContext());
diff --git a/Engine/Plugins/Experimental/PythonScriptPlugin/Source/PythonScriptPlugin/Private/PyMethodWithClosure.cpp b/Engine/Plugins/Experimental/PythonScriptPlugin/Source/PythonScriptPlugin/Private/PyMethodWithClosure.cpp
--- a/Engine/Plugins/Experimental/PythonScriptPlugin/Source/PythonScriptPlugin/Private/PyMethodWithClosure.cpp
+++ b/Engine/Plugins/Experimental/PythonScriptPlugin/Source/PythonScriptPlugin/Private/PyMethodWithClosure.cpp
@@ -77,22 +77,33 @@ void ShutdownPyMethodWithClosure()
 }
 
 
-bool FPyMethodWithClosureDef::AddMethods(FPyMethodWithClosureDef* InMethods, PyTypeObject* InType)
+bool FPyMethodWithClosureDef::AddMethod(FPyMethodWithClosureDef* InMethod, PyTypeObject* InType)
 {
-	for (FPyMethodWithClosureDef* MethodDef = InMethods; MethodDef->MethodName; ++MethodDef)
+	if (!(InMethod->MethodFlags & METH_COEXIST) && PyDict_GetItemString(InType->tp_dict, InMethod->MethodName))
 	{
-		if (!(MethodDef->MethodFlags & METH_COEXIST) && PyDict_GetItemString(InType->tp_dict, MethodDef->MethodName))
-		{
-			continue;
-		}
+		// Skip this method, but don't fail
+		return true;
+	}
 
-		FPyObjectPtr Descr = FPyObjectPtr::StealReference(NewMethodDescriptor(InType, MethodDef));
-		if (!Descr)
-		{
-			return false;
-		}
+	FPyObjectPtr Descr = FPyObjectPtr::StealReference(NewMethodDescriptor(InType, InMethod));
+	if (!Descr)
+	{
+		return false;
+	}
 
-		if (PyDict_SetItemString(InType->tp_dict, MethodDef->MethodName, Descr) != 0)
+	if (PyDict_SetItemString(InType->tp_dict, InMethod->MethodName, Descr) != 0)
+	{
+		return false;
+	}
+
+	return true;
+}
+
+bool FPyMethodWithClosureDef::AddMethods(FPyMethodWithClosureDef* InMethods, PyTypeObject* InType)
+{
+	for (FPyMethodWithClosureDef* MethodDef = InMethods; MethodDef->MethodName; ++MethodDef)
+	{
+		if (!AddMethod(MethodDef, InType))
 		{
 			return false;
 		}
diff --git a/Engine/Plugins/Experimental/PythonScriptPlugin/Source/PythonScriptPlugin/Private/PyMethodWithClosure.h b/Engine/Plugins/Experimental/PythonScriptPlugin/Source/PythonScriptPlugin/Private/PyMethodWithClosure.h
--- a/Engine/Plugins/Experimental/PythonScriptPlugin/Source/PythonScriptPlugin/Private/PyMethodWithClosure.h
+++ b/Engine/Plugins/Experimental/PythonScriptPlugin/Source/PythonScriptPlugin/Private/PyMethodWithClosure.h
@@ -48,6 +48,9 @@ struct FPyMethodWithClosureDef
 	/** The closure passed to the C function, or nullptr */
 	void* MethodClosure;
 
+	/** Add a singular method to the given type */
+	static bool AddMethod(FPyMethodWithClosureDef* InMethod, PyTypeObject* InType);
+
 	/** Add the given null-terminated table of methods to the given type */
 	static bool AddMethods(FPyMethodWithClosureDef* InMethods, PyTypeObject* InType);
 
diff --git a/Engine/Plugins/Experimental/PythonScriptPlugin/Source/PythonScriptPlugin/Private/PyReferenceCollector.cpp b/Engine/Plugins/Experimental/PythonScriptPlugin/Source/PythonScriptPlugin/Private/PyReferenceCollector.cpp
--- a/Engine/Plugins/Experimental/PythonScriptPlugin/Source/PythonScriptPlugin/Private/PyReferenceCollector.cpp
+++ b/Engine/Plugins/Experimental/PythonScriptPlugin/Source/PythonScriptPlugin/Private/PyReferenceCollector.cpp
@@ -78,7 +78,7 @@ void FPyReferenceCollector::AddReferencedObjects(FReferenceCollector& InCollecto
 	for (FPyWrapperBase* PythonWrappedInstance : PythonWrappedInstances)
 	{
 		FPyWrapperBaseMetaData* PythonWrappedInstanceMetaData = FPyWrapperBaseMetaData::GetMetaData(PythonWrappedInstance);
-		if (PythonWrappedInstanceMetaData && PythonWrappedInstanceMetaData->AddReferencedObjects)
+		if (PythonWrappedInstanceMetaData)
 		{
 			PythonWrappedInstanceMetaData->AddReferencedObjects(PythonWrappedInstance, InCollector);
 		}
diff --git a/Engine/Plugins/Experimental/PythonScriptPlugin/Source/PythonScriptPlugin/Private/PySlate.cpp b/Engine/Plugins/Experimental/PythonScriptPlugin/Source/PythonScriptPlugin/Private/PySlate.cpp
new file mode 100644
--- /dev/null
+++ b/Engine/Plugins/Experimental/PythonScriptPlugin/Source/PythonScriptPlugin/Private/PySlate.cpp
@@ -0,0 +1,176 @@
+// Copyright 1998-2018 Epic Games, Inc. All Rights Reserved.
+
+#include "PySlate.h"
+#include "PyCore.h"
+#include "PyUtil.h"
+#include "PyConversion.h"
+
+#include "Framework/Application/SlateApplication.h"
+
+#if PLATFORM_WINDOWS
+#include "Windows/AllowWindowsPlatformTypes.h"
+#include "Windows/HideWindowsPlatformTypes.h"
+#endif
+
+#if WITH_PYTHON
+
+namespace PySlateUtil
+{
+
+FPyWrapperDelegateHandle* RegisterSlateTickCallback(FSlateApplication::FSlateTickEvent& InSlateTickEvent, PyObject* InPyCallable)
+{
+	// We copy the PyCallable into the lambda to keep the Python object alive as long as the delegate is bound
+	FPyObjectPtr PyCallable = FPyObjectPtr::NewReference(InPyCallable);
+	FDelegateHandle TickEventDelegateHandle = InSlateTickEvent.AddLambda([PyCallable](const float InDeltaTime) mutable
+	{
+		FPyObjectPtr PyArgs = FPyObjectPtr::StealReference(PyTuple_New(1));
+		PyTuple_SetItem(PyArgs, 0, PyConversion::Pythonize(InDeltaTime)); // SetItem steals the reference
+
+		FPyObjectPtr Result = FPyObjectPtr::StealReference(PyObject_CallObject(PyCallable.GetPtr(), PyArgs));
+		if (!Result)
+		{
+			PyUtil::LogPythonError();
+		}
+	});
+
+	return FPyWrapperDelegateHandle::CreateInstance(TickEventDelegateHandle);
+}
+
+bool UnregisterSlateTickCallback(FSlateApplication::FSlateTickEvent& InSlateTickEvent, PyObject* InCallbackHandle)
+{
+	FPyWrapperDelegateHandlePtr PyTickEventDelegateHandle = FPyWrapperDelegateHandlePtr::StealReference(FPyWrapperDelegateHandle::CastPyObject(InCallbackHandle));
+	if (!PyTickEventDelegateHandle)
+	{
+		return false;
+	}
+
+	if (PyTickEventDelegateHandle->Value.IsValid())
+	{
+		InSlateTickEvent.Remove(PyTickEventDelegateHandle->Value);
+	}
+
+	return true;
+}
+
+} // namespace PySlateUtil
+
+namespace PySlate
+{
+
+PyObject* RegisterSlatePreTickCallback(PyObject* InSelf, PyObject* InArgs)
+{
+	PyObject* PyObj = nullptr;
+	if (!PyArg_ParseTuple(InArgs, "O:register_slate_pre_tick_callback", &PyObj))
+	{
+		return nullptr;
+	}
+	check(PyObj);
+
+	return (PyObject*)PySlateUtil::RegisterSlateTickCallback(FSlateApplication::Get().OnPreTick(), PyObj);
+}
+
+PyObject* UnregisterSlatePreTickCallback(PyObject* InSelf, PyObject* InArgs)
+{
+	PyObject* PyObj = nullptr;
+	if (!PyArg_ParseTuple(InArgs, "O:unregister_slate_pre_tick_callback", &PyObj))
+	{
+		return nullptr;
+	}
+	check(PyObj);
+
+	if (!PySlateUtil::UnregisterSlateTickCallback(FSlateApplication::Get().OnPreTick(), PyObj))
+	{
+		PyUtil::SetPythonError(PyExc_TypeError, TEXT("unregister_slate_pre_tick_callback"), *FString::Printf(TEXT("Failed to convert argument '%s' to '_DelegateHandle'"), *PyUtil::GetFriendlyTypename(PyObj)));
+		return nullptr;
+	}
+
+	Py_RETURN_NONE;
+}
+
+PyObject* RegisterSlatePostTickCallback(PyObject* InSelf, PyObject* InArgs)
+{
+	PyObject* PyObj = nullptr;
+	if (!PyArg_ParseTuple(InArgs, "O:register_slate_post_tick_callback", &PyObj))
+	{
+		return nullptr;
+	}
+	check(PyObj);
+
+	return (PyObject*)PySlateUtil::RegisterSlateTickCallback(FSlateApplication::Get().OnPostTick(), PyObj);
+}
+
+PyObject* UnregisterSlatePostTickCallback(PyObject* InSelf, PyObject* InArgs)
+{
+	PyObject* PyObj = nullptr;
+	if (!PyArg_ParseTuple(InArgs, "O:unregister_slate_post_tick_callback", &PyObj))
+	{
+		return nullptr;
+	}
+	check(PyObj);
+
+	if (!PySlateUtil::UnregisterSlateTickCallback(FSlateApplication::Get().OnPostTick(), PyObj))
+	{
+		PyUtil::SetPythonError(PyExc_TypeError, TEXT("unregister_slate_post_tick_callback"), *FString::Printf(TEXT("Failed to convert argument '%s' to '_DelegateHandle'"), *PyUtil::GetFriendlyTypename(PyObj)));
+		return nullptr;
+	}
+
+	Py_RETURN_NONE;
+}
+
+PyObject* ParentExternalWindowToSlate(PyObject* InSelf, PyObject* InArgs)
+{
+	PyObject* PyExternalWindowHandle = nullptr;
+	PyObject* PyParentWindowSearchMethod = nullptr;
+	if (!PyArg_ParseTuple(InArgs, "O|O:parent_external_window_to_slate", &PyExternalWindowHandle, &PyParentWindowSearchMethod))
+	{
+		return nullptr;
+	}
+	check(PyExternalWindowHandle);
+
+	void* ExternalWindowHandle = nullptr;
+	if (!PyConversion::Nativize(PyExternalWindowHandle, ExternalWindowHandle))
+	{
+		PyUtil::SetPythonError(PyExc_TypeError, TEXT("parent_external_window_to_slate"), *FString::Printf(TEXT("Failed to convert argument '%s' to 'void*'"), *PyUtil::GetFriendlyTypename(PyExternalWindowHandle)));
+		return nullptr;
+	}
+
+	ESlateParentWindowSearchMethod ParentWindowSearchMethod = ESlateParentWindowSearchMethod::ActiveWindow;
+	if (PyParentWindowSearchMethod && !PyConversion::NativizeEnum(PyParentWindowSearchMethod, ParentWindowSearchMethod))
+	{
+		PyUtil::SetPythonError(PyExc_TypeError, TEXT("parent_external_window_to_slate"), *FString::Printf(TEXT("Failed to convert argument '%s' to 'SlateParentWindowSearchMethod'"), *PyUtil::GetFriendlyTypename(PyParentWindowSearchMethod)));
+		return nullptr;
+	}
+
+	const void* SlateParentWindowHandle = FSlateApplication::Get().FindBestParentWindowHandleForDialogs(nullptr, ParentWindowSearchMethod);
+	if (SlateParentWindowHandle && ExternalWindowHandle)
+	{
+#if PLATFORM_WINDOWS
+		::SetWindowLongPtr((HWND)ExternalWindowHandle, -8/*GWL_HWNDPARENT*/, (LONG_PTR)SlateParentWindowHandle);
+#endif
+	}
+
+	Py_RETURN_NONE;
+}
+
+PyMethodDef PySlateMethods[] = {
+	{ "register_slate_pre_tick_callback", PyCFunctionCast(&RegisterSlatePreTickCallback), METH_VARARGS, "x.register_slate_pre_tick_callback(callable) -> _DelegateHandle -- register the given callable (taking a single float) as a pre-tick callback in Slate" },
+	{ "unregister_slate_pre_tick_callback", PyCFunctionCast(&UnregisterSlatePreTickCallback), METH_VARARGS, "x.unregister_slate_pre_tick_callback(handle) -- unregister the given handle from a previous call to register_slate_pre_tick_callback" },
+	{ "register_slate_post_tick_callback", PyCFunctionCast(&RegisterSlatePostTickCallback), METH_VARARGS, "x.register_slate_post_tick_callback(callable) -> _DelegateHandle -- register the given callable (taking a single float) as a pre-tick callback in Slate" },
+	{ "unregister_slate_post_tick_callback", PyCFunctionCast(&UnregisterSlatePostTickCallback), METH_VARARGS, "x.unregister_slate_post_tick_callback(handle) -- unregister the given handle from a previous call to register_slate_post_tick_callback" },
+	{ "parent_external_window_to_slate", PyCFunctionCast(&ParentExternalWindowToSlate), METH_VARARGS, "x.parent_external_window_to_slate(external_window, parent_search_method=SlateParentWindowSearchMethod.ACTIVE_WINDOW) -- parent the given OS specific external window handle to a suitable Slate window" },
+	{ nullptr, nullptr, 0, nullptr }
+};
+
+void InitializeModule()
+{
+#if PY_MAJOR_VERSION >= 3
+	PyObject* PyModule = PyImport_AddModule("_unreal_slate");
+	PyModule_AddFunctions(PyModule, PySlateMethods);
+#else	// PY_MAJOR_VERSION >= 3
+	PyObject* PyModule = Py_InitModule("_unreal_slate", PySlateMethods);
+#endif	// PY_MAJOR_VERSION >= 3
+}
+
+}
+
+#endif	// WITH_PYTHON
diff --git a/Engine/Plugins/Experimental/PythonScriptPlugin/Source/PythonScriptPlugin/Private/PyTest.cpp b/Engine/Plugins/Experimental/PythonScriptPlugin/Source/PythonScriptPlugin/Private/PyTest.cpp
--- a/Engine/Plugins/Experimental/PythonScriptPlugin/Source/PythonScriptPlugin/Private/PyTest.cpp
+++ b/Engine/Plugins/Experimental/PythonScriptPlugin/Source/PythonScriptPlugin/Private/PyTest.cpp
@@ -11,6 +11,37 @@ FPyTestStruct::FPyTestStruct()
 	Enum = EPyTestEnum::One;
 }
 
+bool UPyTestStructLibrary::IsBoolSet(const FPyTestStruct& InStruct)
+{
+	return InStruct.Bool;
+}
+
+bool UPyTestStructLibrary::LegacyIsBoolSet(const FPyTestStruct& InStruct)
+{
+	return IsBoolSet(InStruct);
+}
+
+FPyTestStruct UPyTestStructLibrary::AddInt(const FPyTestStruct& InStruct, const int32 InValue)
+{
+	FPyTestStruct Result = InStruct;
+	Result.Int += InValue;
+	return Result;
+}
+
+FPyTestStruct UPyTestStructLibrary::AddFloat(const FPyTestStruct& InStruct, const float InValue)
+{
+	FPyTestStruct Result = InStruct;
+	Result.Float += InValue;
+	return Result;
+}
+
+FPyTestStruct UPyTestStructLibrary::AddStr(const FPyTestStruct& InStruct, const FString& InValue)
+{
+	FPyTestStruct Result = InStruct;
+	Result.String += InValue;
+	return Result;
+}
+
 UPyTestObject::UPyTestObject()
 {
 	StructArray.AddDefaulted();
@@ -40,6 +71,11 @@ void UPyTestObject::FuncTakingPyTestChildStruct(const FPyTestChildStruct& InStru
 {
 }
 
+void UPyTestObject::LegacyFuncTakingPyTestStruct(const FPyTestStruct& InStruct) const
+{
+	FuncTakingPyTestStruct(InStruct);
+}
+
 int32 UPyTestObject::FuncTakingPyTestDelegate(const FPyTestDelegate& InDelegate, const int32 InValue) const
 {
 	return InDelegate.IsBound() ? InDelegate.Execute(InValue) : INDEX_NONE;
diff --git a/Engine/Plugins/Experimental/PythonScriptPlugin/Source/PythonScriptPlugin/Private/PyTest.h b/Engine/Plugins/Experimental/PythonScriptPlugin/Source/PythonScriptPlugin/Private/PyTest.h
--- a/Engine/Plugins/Experimental/PythonScriptPlugin/Source/PythonScriptPlugin/Private/PyTest.h
+++ b/Engine/Plugins/Experimental/PythonScriptPlugin/Source/PythonScriptPlugin/Private/PyTest.h
@@ -3,6 +3,7 @@
 #pragma once
 
 #include "CoreMinimal.h"
+#include "Kismet/BlueprintFunctionLibrary.h"
 #include "PyTest.generated.h"
 
 /**
@@ -65,6 +66,9 @@ public:
 
 	UPROPERTY(EditAnywhere, BlueprintReadWrite, Category = Python)
 	TMap<FString, int32> StringIntMap;
+
+	UPROPERTY(meta=(DeprecatedProperty, DeprecationMessage="LegacyInt is deprecated. Please use Int instead."))
+	int32 LegacyInt_DEPRECATED;
 };
 
 /**
@@ -76,6 +80,30 @@ struct FPyTestChildStruct : public FPyTestStruct
 	GENERATED_BODY()
 };
 
+/**
+ * Function library containing methods that should be hoisted onto the test struct in Python.
+ */
+UCLASS()
+class UPyTestStructLibrary : public UBlueprintFunctionLibrary
+{
+	GENERATED_BODY()
+
+	UFUNCTION(BlueprintPure, Category = Python, meta=(ScriptMethod))
+	static bool IsBoolSet(const FPyTestStruct& InStruct);
+
+	UFUNCTION(BlueprintPure, Category = Python, meta=(ScriptMethod, DeprecatedFunction, DeprecationMessage="LegacyIsBoolSet is deprecated. Please use IsBoolSet instead."))
+	static bool LegacyIsBoolSet(const FPyTestStruct& InStruct);
+
+	UFUNCTION(BlueprintPure, Category = Python, meta=(ScriptMathOp="+;+="))
+	static FPyTestStruct AddInt(const FPyTestStruct& InStruct, const int32 InValue);
+
+	UFUNCTION(BlueprintPure, Category = Python, meta=(ScriptMathOp="+;+="))
+	static FPyTestStruct AddFloat(const FPyTestStruct& InStruct, const float InValue);
+
+	UFUNCTION(BlueprintPure, Category = Python, meta=(ScriptMathOp="+;+="))
+	static FPyTestStruct AddStr(const FPyTestStruct& InStruct, const FString& InValue);
+};
+
 /**
  * Object to allow testing of the various UObject features that are exposed to Python wrapped types.
  */
@@ -150,6 +178,9 @@ public:
 	UFUNCTION(BlueprintCallable, Category = Python)
 	void FuncTakingPyTestChildStruct(const FPyTestChildStruct& InStruct) const;
 
+	UFUNCTION(BlueprintCallable, Category = Python, meta=(DeprecatedFunction, DeprecationMessage="LegacyFuncTakingPyTestStruct is deprecated. Please use FuncTakingPyTestStruct instead."))
+	void LegacyFuncTakingPyTestStruct(const FPyTestStruct& InStruct) const;
+
 	UFUNCTION(BlueprintCallable, Category = Python)
 	int32 FuncTakingPyTestDelegate(const FPyTestDelegate& InDelegate, const int32 InValue) const;
 
@@ -168,3 +199,12 @@ class UPyTestChildObject : public UPyTestObject
 {
 	GENERATED_BODY()
 };
+
+/**
+ * Object to test deprecation of Python wrapped types.
+ */
+UCLASS(Blueprintable, deprecated, meta=(DeprecationMessage="LegacyPyTestObject is deprecated. Please use PyTestObject instead."))
+class UDEPRECATED_LegacyPyTestObject : public UPyTestObject
+{
+	GENERATED_BODY()
+};
diff --git a/Engine/Plugins/Experimental/PythonScriptPlugin/Source/PythonScriptPlugin/Private/PyUtil.cpp b/Engine/Plugins/Experimental/PythonScriptPlugin/Source/PythonScriptPlugin/Private/PyUtil.cpp
--- a/Engine/Plugins/Experimental/PythonScriptPlugin/Source/PythonScriptPlugin/Private/PyUtil.cpp
+++ b/Engine/Plugins/Experimental/PythonScriptPlugin/Source/PythonScriptPlugin/Private/PyUtil.cpp
@@ -233,14 +233,14 @@ bool CalculatePropertyDef(PyTypeObject* InPyType, FPropertyDef& OutPropertyDef)
 	if (PyObject_IsSubclass((PyObject*)InPyType, (PyObject*)&PyWrapperDelegateType) == 1)
 	{
 		OutPropertyDef.PropertyClass = UDelegateProperty::StaticClass();
-		OutPropertyDef.PropertySubType = (UObject*)FPyWrapperDelegateMetaData::GetDelegateSignature(InPyType);
+		OutPropertyDef.PropertySubType = (UObject*)FPyWrapperDelegateMetaData::GetDelegateSignature(InPyType).Func;
 		return true;
 	}
 
 	if (PyObject_IsSubclass((PyObject*)InPyType, (PyObject*)&PyWrapperMulticastDelegateType) == 1)
 	{
 		OutPropertyDef.PropertyClass = UMulticastDelegateProperty::StaticClass();
-		OutPropertyDef.PropertySubType = (UObject*)FPyWrapperMulticastDelegateMetaData::GetDelegateSignature(InPyType);
+		OutPropertyDef.PropertySubType = (UObject*)FPyWrapperMulticastDelegateMetaData::GetDelegateSignature(InPyType).Func;
 		return true;
 	}
 
@@ -434,115 +434,10 @@ bool IsOutputParameter(const UProperty* InParam)
 	return !bIsReturnParam && bIsOutParam;
 }
 
-PyObject* PackReturnValues(const UFunction* InFunc, const TArrayView<const UProperty*>& InReturnProperties, const void* InBaseParamsAddr, const TCHAR* InErrorCtxt, const TCHAR* InCallingCtxt)
+void InvokeFunctionCall(UObject* InObj, const UFunction* InFunc, void* InBaseParamsAddr, const TCHAR* InErrorCtxt)
 {
-	if (!InReturnProperties.Num())
-	{
-		Py_RETURN_NONE;
-	}
-
-	int32 ReturnPropIndex = 0;
-
-	// If we have multiple return values and the main return value is a bool, we return None (for false) or the (potentially packed) return value without the bool (for true)
-	if (InReturnProperties.Num() > 1 && InReturnProperties[0]->IsA<UBoolProperty>())
-	{
-		const UBoolProperty* BoolReturn = CastChecked<const UBoolProperty>(InReturnProperties[0]);
-		const bool bReturnValue = BoolReturn->GetPropertyValue(BoolReturn->ContainerPtrToValuePtr<void>(InBaseParamsAddr));
-		if (!bReturnValue)
-		{
-			Py_RETURN_NONE;
-		}
-
-		ReturnPropIndex = 1; // Start packing at the 1st out value
-	}
-
-	// Do we need to return a packed tuple, or just a single value?
-	const int32 NumPropertiesToPack = InReturnProperties.Num() - ReturnPropIndex;
-	if (NumPropertiesToPack == 1)
-	{
-		PyObject* OutParamPyObj = nullptr;
-		if (!PyConversion::PythonizeProperty_InContainer(InReturnProperties[ReturnPropIndex], InBaseParamsAddr, 0, OutParamPyObj, EPyConversionMethod::Steal))
-		{
-			PyUtil::SetPythonError(PyExc_TypeError, InErrorCtxt, *FString::Printf(TEXT("Failed to convert return property '%s' (%s) when calling %s"), *InReturnProperties[ReturnPropIndex]->GetName(), *InReturnProperties[ReturnPropIndex]->GetClass()->GetName(), InCallingCtxt));
-			return nullptr;
-		}
-		return OutParamPyObj;
-	}
-	else
-	{
-		int32 OutParamTupleIndex = 0;
-		FPyObjectPtr OutParamTuple = FPyObjectPtr::StealReference(PyTuple_New(NumPropertiesToPack));
-		for (; ReturnPropIndex < InReturnProperties.Num(); ++ReturnPropIndex)
-		{
-			PyObject* OutParamPyObj = nullptr;
-			if (!PyConversion::PythonizeProperty_InContainer(InReturnProperties[ReturnPropIndex], InBaseParamsAddr, 0, OutParamPyObj, EPyConversionMethod::Steal))
-			{
-				PyUtil::SetPythonError(PyExc_TypeError, InErrorCtxt, *FString::Printf(TEXT("Failed to convert return property '%s' (%s) when calling function %s"), *InReturnProperties[ReturnPropIndex]->GetName(), *InReturnProperties[ReturnPropIndex]->GetClass()->GetName(), InCallingCtxt));
-				return nullptr;
-			}
-			PyTuple_SetItem(OutParamTuple, OutParamTupleIndex++, OutParamPyObj); // SetItem steals the reference
-		}
-		return OutParamTuple.Release();
-	}
-}
-
-bool UnpackReturnValues(PyObject* InRetVals, const UFunction* InFunc, const TArrayView<const UProperty*>& InReturnProperties, void* InBaseParamsAddr, const TCHAR* InErrorCtxt, const TCHAR* InCallingCtxt)
-{
-	if (!InReturnProperties.Num())
-	{
-		return true;
-	}
-
-	int32 ReturnPropIndex = 0;
-
-	// If we have multiple return values and the main return value is a bool, we expect None (for false) or the (potentially packed) return value without the bool (for true)
-	if (InReturnProperties.Num() > 1 && InReturnProperties[0]->IsA<UBoolProperty>())
-	{
-		const UBoolProperty* BoolReturn = CastChecked<const UBoolProperty>(InReturnProperties[0]);
-		const bool bReturnValue = InRetVals != Py_None;
-		BoolReturn->SetPropertyValue(BoolReturn->ContainerPtrToValuePtr<void>(InBaseParamsAddr), bReturnValue);
-
-		ReturnPropIndex = 1; // Start unpacking at the 1st out value
-	}
-
-	// Do we need to expect a packed tuple, or just a single value?
-	const int32 NumPropertiesToUnpack = InReturnProperties.Num() - ReturnPropIndex;
-	if (NumPropertiesToUnpack == 1)
-	{
-		if (!PyConversion::NativizeProperty_InContainer(InRetVals, InReturnProperties[ReturnPropIndex], InBaseParamsAddr, 0))
-		{
-			PyUtil::SetPythonError(PyExc_TypeError, InErrorCtxt, *FString::Printf(TEXT("Failed to convert return property '%s' (%s) when calling %s"), *InReturnProperties[ReturnPropIndex]->GetName(), *InReturnProperties[ReturnPropIndex]->GetClass()->GetName(), InCallingCtxt));
-			return false;
-		}
-	}
-	else
-	{
-		if (!PyTuple_Check(InRetVals))
-		{
-			PyUtil::SetPythonError(PyExc_TypeError, InErrorCtxt, *FString::Printf(TEXT("Expected a 'tuple' return type, but got '%s' when calling %s"), *GetFriendlyTypename(InRetVals), InCallingCtxt));
-			return false;
-		}
-
-		const int32 RetTupleSize = PyTuple_Size(InRetVals);
-		if (RetTupleSize != NumPropertiesToUnpack)
-		{
-			PyUtil::SetPythonError(PyExc_TypeError, InErrorCtxt, *FString::Printf(TEXT("Expected a 'tuple' return type containing '%d' items but got one containing '%d' items when calling %s"), NumPropertiesToUnpack, RetTupleSize, InCallingCtxt));
-			return false;
-		}
-
-		int32 RetTupleIndex = 0;
-		for (; ReturnPropIndex < InReturnProperties.Num(); ++ReturnPropIndex)
-		{
-			PyObject* RetVal = PyTuple_GetItem(InRetVals, RetTupleIndex++);
-			if (!PyConversion::NativizeProperty_InContainer(RetVal, InReturnProperties[ReturnPropIndex], InBaseParamsAddr, 0))
-			{
-				PyUtil::SetPythonError(PyExc_TypeError, InErrorCtxt, *FString::Printf(TEXT("Failed to convert return property '%s' (%s) when calling %s"), *InReturnProperties[ReturnPropIndex]->GetName(), *InReturnProperties[ReturnPropIndex]->GetClass()->GetName(), InCallingCtxt));
-				return false;
-			}
-		}
-	}
-
-	return true;
+	FEditorScriptExecutionGuard ScriptGuard;
+	InObj->ProcessEvent((UFunction*)InFunc, InBaseParamsAddr);
 }
 
 bool InspectFunctionArgs(PyObject* InFunc, TArray<FString>& OutArgNames, TArray<FPyObjectPtr>* OutArgDefaults)
@@ -657,27 +552,30 @@ int ValidateContainerIndexParam(const Py_ssize_t InIndex, const Py_ssize_t InLen
 	return 0;
 }
 
-PyObject* GetUEPropValue(const UStruct* InStruct, void* InStructData, const FName InPropName, const char *InAttributeName, PyObject* InOwnerPyObject, const TCHAR* InErrorCtxt)
+PyObject* GetUEPropValue(const UStruct* InStruct, void* InStructData, const UProperty* InProp, const char *InAttributeName, PyObject* InOwnerPyObject, const TCHAR* InErrorCtxt)
 {
-	if (InStruct && ensureAlways(InStructData))
+	if (InStruct && InProp && ensureAlways(InStructData))
 	{
-		UProperty* Prop = InStruct->FindPropertyByName(InPropName);
-		if (!Prop)
+		// Deprecated properties can no longer be accessed and cause an error rather than a warning
 		{
-			SetPythonError(PyExc_Exception, InErrorCtxt, *FString::Printf(TEXT("Failed to find property '%s' for attribute '%s' on '%s'"), *InPropName.ToString(), UTF8_TO_TCHAR(InAttributeName), *InStruct->GetName()));
-			return nullptr;
+			FString DeprecationMessage;
+			if (PyGenUtil::IsDeprecatedProperty(InProp, &DeprecationMessage))
+			{
+				SetPythonError(PyExc_DeprecationWarning, InErrorCtxt, *FString::Printf(TEXT("Property '%s' for attribute '%s' on '%s' is deprecated: %s"), *InProp->GetName(), UTF8_TO_TCHAR(InAttributeName), *InStruct->GetName(), *DeprecationMessage));
+				return nullptr;
+			}
 		}
 
-		if (!Prop->HasAnyPropertyFlags(CPF_Edit | CPF_BlueprintVisible))
+		if (!InProp->HasAnyPropertyFlags(CPF_Edit | CPF_BlueprintVisible))
 		{
-			SetPythonError(PyExc_Exception, InErrorCtxt, *FString::Printf(TEXT("Property '%s' for attribute '%s' on '%s' is protected and cannot be read"), *Prop->GetName(), UTF8_TO_TCHAR(InAttributeName), *InStruct->GetName()));
+			SetPythonError(PyExc_Exception, InErrorCtxt, *FString::Printf(TEXT("Property '%s' for attribute '%s' on '%s' is protected and cannot be read"), *InProp->GetName(), UTF8_TO_TCHAR(InAttributeName), *InStruct->GetName()));
 			return nullptr;
 		}
 
 		PyObject* PyPropObj = nullptr;
-		if (!PyConversion::PythonizeProperty_InContainer(Prop, InStructData, 0, PyPropObj, EPyConversionMethod::Reference, InOwnerPyObject))
+		if (!PyConversion::PythonizeProperty_InContainer(InProp, InStructData, 0, PyPropObj, EPyConversionMethod::Reference, InOwnerPyObject))
 		{
-			SetPythonError(PyExc_TypeError, InErrorCtxt, *FString::Printf(TEXT("Failed to convert property '%s' (%s) for attribute '%s' on '%s'"), *Prop->GetName(), *Prop->GetClass()->GetName(), UTF8_TO_TCHAR(InAttributeName), *InStruct->GetName()));
+			SetPythonError(PyExc_TypeError, InErrorCtxt, *FString::Printf(TEXT("Failed to convert property '%s' (%s) for attribute '%s' on '%s'"), *InProp->GetName(), *InProp->GetClass()->GetName(), UTF8_TO_TCHAR(InAttributeName), *InStruct->GetName()));
 			return nullptr;
 		}
 		return PyPropObj;
@@ -686,7 +584,7 @@ PyObject* GetUEPropValue(const UStruct* InStruct, void* InStructData, const FNam
 	Py_RETURN_NONE;
 }
 
-int SetUEPropValue(const UStruct* InStruct, void* InStructData, PyObject* InValue, const FName InPropName, const char *InAttributeName, const FPyWrapperOwnerContext& InChangeOwner, const uint64 InReadOnlyFlags, const TCHAR* InErrorCtxt)
+int SetUEPropValue(const UStruct* InStruct, void* InStructData, PyObject* InValue, const UProperty* InProp, const char *InAttributeName, const FPyWrapperOwnerContext& InChangeOwner, const uint64 InReadOnlyFlags, const TCHAR* InErrorCtxt)
 {
 	if (!InValue)
 	{
@@ -694,30 +592,33 @@ int SetUEPropValue(const UStruct* InStruct, void* InStructData, PyObject* InValu
 		return -1;
 	}
 
-	if (InStruct && ensureAlways(InStructData))
+	if (InStruct && InProp && ensureAlways(InStructData))
 	{
-		UProperty* Prop = InStruct->FindPropertyByName(InPropName);
-		if (!Prop)
+		// Deprecated properties can no longer be accessed and cause an error rather than a warning
 		{
-			SetPythonError(PyExc_Exception, InErrorCtxt, *FString::Printf(TEXT("Failed to find property '%s' for attribute '%s' on '%s'"), *InPropName.ToString(), UTF8_TO_TCHAR(InAttributeName), *InStruct->GetName()));
-			return -1;
+			FString DeprecationMessage;
+			if (PyGenUtil::IsDeprecatedProperty(InProp, &DeprecationMessage))
+			{
+				SetPythonError(PyExc_DeprecationWarning, InErrorCtxt, *FString::Printf(TEXT("Property '%s' for attribute '%s' on '%s' is deprecated: %s"), *InProp->GetName(), UTF8_TO_TCHAR(InAttributeName), *InStruct->GetName(), *DeprecationMessage));
+				return -1;
+			}
 		}
 
-		if (!Prop->HasAnyPropertyFlags(CPF_Edit | CPF_BlueprintVisible))
+		if (!InProp->HasAnyPropertyFlags(CPF_Edit | CPF_BlueprintVisible))
 		{
-			SetPythonError(PyExc_Exception, InErrorCtxt, *FString::Printf(TEXT("Property '%s' for attribute '%s' on '%s' is protected and cannot be set"), *Prop->GetName(), UTF8_TO_TCHAR(InAttributeName), *InStruct->GetName()));
+			SetPythonError(PyExc_Exception, InErrorCtxt, *FString::Printf(TEXT("Property '%s' for attribute '%s' on '%s' is protected and cannot be set"), *InProp->GetName(), UTF8_TO_TCHAR(InAttributeName), *InStruct->GetName()));
 			return -1;
 		}
 
-		if (Prop->HasAnyPropertyFlags(InReadOnlyFlags))
+		if (InProp->HasAnyPropertyFlags(InReadOnlyFlags))
 		{
-			SetPythonError(PyExc_Exception, InErrorCtxt, *FString::Printf(TEXT("Property '%s' for attribute '%s' on '%s' is read-only and cannot be set"), *Prop->GetName(), UTF8_TO_TCHAR(InAttributeName), *InStruct->GetName()));
+			SetPythonError(PyExc_Exception, InErrorCtxt, *FString::Printf(TEXT("Property '%s' for attribute '%s' on '%s' is read-only and cannot be set"), *InProp->GetName(), UTF8_TO_TCHAR(InAttributeName), *InStruct->GetName()));
 			return -1;
 		}
 
-		if (!PyConversion::NativizeProperty_InContainer(InValue, Prop, InStructData, 0, InChangeOwner))
+		if (!PyConversion::NativizeProperty_InContainer(InValue, InProp, InStructData, 0, InChangeOwner))
 		{
-			SetPythonError(PyExc_TypeError, InErrorCtxt, *FString::Printf(TEXT("Failed to convert type '%s' to property '%s' (%s) for attribute '%s' on '%s'"), *GetFriendlyTypename(InValue), *Prop->GetName(), *Prop->GetClass()->GetName(), UTF8_TO_TCHAR(InAttributeName), *InStruct->GetName()));
+			SetPythonError(PyExc_TypeError, InErrorCtxt, *FString::Printf(TEXT("Failed to convert type '%s' to property '%s' (%s) for attribute '%s' on '%s'"), *GetFriendlyTypename(InValue), *InProp->GetName(), *InProp->GetClass()->GetName(), UTF8_TO_TCHAR(InAttributeName), *InStruct->GetName()));
 			return -1;
 		}
 	}
@@ -917,14 +818,7 @@ FString GetFriendlyStructValue(const UStruct* InStruct, const void* InStructValu
 			FriendlyStructValue += UTF8_TO_TCHAR(InitParam.ParamName.GetData());
 			FriendlyStructValue += TEXT(": ");
 
-			if (UProperty* Prop = InStruct->FindPropertyByName(InitParam.ParamPropName))
-			{
-				FriendlyStructValue += GetFriendlyPropertyValue(Prop, Prop->ContainerPtrToValuePtr<void>(InStructValue), InPortFlags | PPF_Delimited);
-			}
-			else
-			{
-				FriendlyStructValue += TEXT("<missing property>");
-			}
+			FriendlyStructValue += GetFriendlyPropertyValue(InitParam.ParamProp, InitParam.ParamProp->ContainerPtrToValuePtr<void>(InStructValue), InPortFlags | PPF_Delimited);
 		}
 
 		FriendlyStructValue += TEXT('}');
@@ -1068,6 +962,43 @@ void SetPythonError(PyObject* InException, const TCHAR* InErrorContext, const TC
 	PyErr_SetString(InException, TCHAR_TO_UTF8(*FinalException));
 }
 
+int SetPythonWarning(PyObject* InException, PyTypeObject* InErrorContext, const TCHAR* InErrorMsg)
+{
+	return SetPythonWarning(InException, *GetErrorContext(InErrorContext), InErrorMsg);
+}
+
+int SetPythonWarning(PyObject* InException, PyObject* InErrorContext, const TCHAR* InErrorMsg)
+{
+	return SetPythonWarning(InException, *GetErrorContext(InErrorContext), InErrorMsg);
+}
+
+int SetPythonWarning(PyObject* InException, const TCHAR* InErrorContext, const TCHAR* InErrorMsg)
+{
+	const FString FinalException = FString::Printf(TEXT("%s: %s"), InErrorContext, InErrorMsg);
+	return PyErr_WarnEx(InException, TCHAR_TO_UTF8(*FinalException), 1);
+}
+
+bool EnableDeveloperWarnings()
+{
+	FPyObjectPtr PyWarningsModule = FPyObjectPtr::StealReference(PyImport_ImportModule("warnings"));
+	if (PyWarningsModule)
+	{
+		PyObject* PyWarningsDict = PyModule_GetDict(PyWarningsModule);
+
+		PyObject* PySimpleFilterFunc = PyDict_GetItemString(PyWarningsDict, "simplefilter");
+		if (PySimpleFilterFunc)
+		{
+			FPyObjectPtr PySimpleFilterResult = FPyObjectPtr::StealReference(PyObject_CallFunction(PySimpleFilterFunc, "s", "default"));
+			if (PySimpleFilterResult)
+			{
+				return true;
+			}
+		}
+	}
+
+	return false;
+}
+
 void LogPythonError(const bool bInteractive)
 {
 	auto BuildPythonError = []() -> FString
diff --git a/Engine/Plugins/Experimental/PythonScriptPlugin/Source/PythonScriptPlugin/Private/PyUtil.h b/Engine/Plugins/Experimental/PythonScriptPlugin/Source/PythonScriptPlugin/Private/PyUtil.h
--- a/Engine/Plugins/Experimental/PythonScriptPlugin/Source/PythonScriptPlugin/Private/PyUtil.h
+++ b/Engine/Plugins/Experimental/PythonScriptPlugin/Source/PythonScriptPlugin/Private/PyUtil.h
@@ -206,11 +206,8 @@ namespace PyUtil
 	/** Check to see if the given property is an output parameter for a function */
 	bool IsOutputParameter(const UProperty* InParam);
 
-	/** Given a set of potential return values and the struct data associated with them, pack them appropriately for returning to Python */
-	PyObject* PackReturnValues(const UFunction* InFunc, const TArrayView<const UProperty*>& InReturnProperties, const void* InBaseParamsAddr, const TCHAR* InErrorCtxt, const TCHAR* InCallingCtxt);
-
-	/** Given a Python return value, unpack the values into the struct data associated with them */
-	bool UnpackReturnValues(PyObject* InRetVals, const UFunction* InFunc, const TArrayView<const UProperty*>& InReturnProperties, void* InBaseParamsAddr, const TCHAR* InErrorCtxt, const TCHAR* InCallingCtxt);
+	/** Invoke a function call */
+	void InvokeFunctionCall(UObject* InObj, const UFunction* InFunc, void* InBaseParamsAddr, const TCHAR* InErrorCtxt);
 
 	/** Given a Python function, get the names of the arguments along with their default values */
 	bool InspectFunctionArgs(PyObject* InFunc, TArray<FString>& OutArgNames, TArray<FPyObjectPtr>* OutArgDefaults = nullptr);
@@ -225,10 +222,10 @@ namespace PyUtil
 	int ValidateContainerIndexParam(const Py_ssize_t InIndex, const Py_ssize_t InLen, const UProperty* InProp, const TCHAR* InErrorCtxt);
 
 	/** Given a struct, get the named property value from it */
-	PyObject* GetUEPropValue(const UStruct* InStruct, void* InStructData, const FName InPropName, const char *InAttributeName, PyObject* InOwnerPyObject, const TCHAR* InErrorCtxt);
+	PyObject* GetUEPropValue(const UStruct* InStruct, void* InStructData, const UProperty* InProp, const char *InAttributeName, PyObject* InOwnerPyObject, const TCHAR* InErrorCtxt);
 
 	/** Given a struct, set the named property value on it */
-	int SetUEPropValue(const UStruct* InStruct, void* InStructData, PyObject* InValue, const FName InPropName, const char *InAttributeName, const FPyWrapperOwnerContext& InChangeOwner, const uint64 InReadOnlyFlags, const TCHAR* InErrorCtxt);
+	int SetUEPropValue(const UStruct* InStruct, void* InStructData, PyObject* InValue, const UProperty* InProp, const char *InAttributeName, const FPyWrapperOwnerContext& InChangeOwner, const uint64 InReadOnlyFlags, const TCHAR* InErrorCtxt);
 
 	/**
 	 * Check to see if the given object implements a length function.
@@ -348,6 +345,20 @@ namespace PyUtil
 		return SetPythonError(InException, (PyObject*)InErrorContext, InErrorMsg);
 	}
 
+	/** Set a Python warning (see PyErr_WarnEx for the return value meaning) */
+	int SetPythonWarning(PyObject* InException, PyTypeObject* InErrorContext, const TCHAR* InErrorMsg);
+	int SetPythonWarning(PyObject* InException, PyObject* InErrorContext, const TCHAR* InErrorMsg);
+	int SetPythonWarning(PyObject* InException, const TCHAR* InErrorContext, const TCHAR* InErrorMsg);
+
+	template <typename T>
+	int SetPythonWarning(PyObject* InException, T* InErrorContext, const TCHAR* InErrorMsg)
+	{
+		return SetPythonWarning(InException, (PyObject*)InErrorContext, InErrorMsg);
+	}
+
+	/** Enable developer warnings (eg, deprecation warnings) */
+	bool EnableDeveloperWarnings();
+
 	/** Log any pending Python error (will also clear the error) */
 	void LogPythonError(const bool bInteractive = false);
 }
diff --git a/Engine/Plugins/Experimental/PythonScriptPlugin/Source/PythonScriptPlugin/Private/PyWrapperArray.cpp b/Engine/Plugins/Experimental/PythonScriptPlugin/Source/PythonScriptPlugin/Private/PyWrapperArray.cpp
--- a/Engine/Plugins/Experimental/PythonScriptPlugin/Source/PythonScriptPlugin/Private/PyWrapperArray.cpp
+++ b/Engine/Plugins/Experimental/PythonScriptPlugin/Source/PythonScriptPlugin/Private/PyWrapperArray.cpp
@@ -279,7 +279,7 @@ bool FPyWrapperArray::ValidateInternalState(FPyWrapperArray* InSelf)
 	return true;
 }
 
-FPyWrapperArray* FPyWrapperArray::CastPyObject(PyObject* InPyObject)
+FPyWrapperArray* FPyWrapperArray::CastPyObject(PyObject* InPyObject, FPyConversionResult* OutCastResult)
 {
 	if (PyObject_IsInstance(InPyObject, (PyObject*)&PyWrapperArrayType) == 1)
 	{
@@ -290,8 +290,10 @@ FPyWrapperArray* FPyWrapperArray::CastPyObject(PyObject* InPyObject)
 	return nullptr;
 }
 
-FPyWrapperArray* FPyWrapperArray::CastPyObject(PyObject* InPyObject, PyTypeObject* InType, const PyUtil::FPropertyDef& InElementDef)
+FPyWrapperArray* FPyWrapperArray::CastPyObject(PyObject* InPyObject, PyTypeObject* InType, const PyUtil::FPropertyDef& InElementDef, FPyConversionResult* OutCastResult)
 {
+	SetOptionalPyConversionResult(FPyConversionResult::Failure(), OutCastResult);
+
 	if (PyObject_IsInstance(InPyObject, (PyObject*)InType) == 1 && (InType == &PyWrapperArrayType || PyObject_IsInstance(InPyObject, (PyObject*)&PyWrapperArrayType) == 1))
 	{
 		FPyWrapperArray* Self = (FPyWrapperArray*)InPyObject;
@@ -304,6 +306,8 @@ FPyWrapperArray* FPyWrapperArray::CastPyObject(PyObject* InPyObject, PyTypeObjec
 		const PyUtil::FPropertyDef SelfElementPropDef = Self->ArrayProp->Inner;
 		if (SelfElementPropDef == InElementDef)
 		{
+			SetOptionalPyConversionResult(FPyConversionResult::Success(), OutCastResult);
+
 			Py_INCREF(Self);
 			return Self;
 		}
@@ -340,6 +344,7 @@ FPyWrapperArray* FPyWrapperArray::CastPyObject(PyObject* InPyObject, PyTypeObjec
 			}
 		}
 
+		SetOptionalPyConversionResult(FPyConversionResult::SuccessWithCoercion(), OutCastResult);
 		return NewArray.Release();
 	}
 
@@ -375,6 +380,7 @@ FPyWrapperArray* FPyWrapperArray::CastPyObject(PyObject* InPyObject, PyTypeObjec
 					}
 				}
 
+				SetOptionalPyConversionResult(FPyConversionResult::SuccessWithCoercion(), OutCastResult);
 				return NewArray.Release();
 			}
 		}
@@ -1505,17 +1511,14 @@ PyTypeObject InitializePyWrapperArrayIteratorType()
 PyTypeObject PyWrapperArrayType = InitializePyWrapperArrayType();
 PyTypeObject PyWrapperArrayIteratorType = InitializePyWrapperArrayIteratorType();
 
-FPyWrapperArrayMetaData::FPyWrapperArrayMetaData()
+void FPyWrapperArrayMetaData::AddReferencedObjects(FPyWrapperBase* Instance, FReferenceCollector& Collector)
 {
-	AddReferencedObjects = [](FPyWrapperBase* Self, FReferenceCollector& Collector)
+	FPyWrapperArray* Self = static_cast<FPyWrapperArray*>(Instance);
+	if (Self->ArrayProp && Self->ArrayInstance && !Self->OwnerContext.HasOwner())
 	{
-		FPyWrapperArray* ArraySelf = static_cast<FPyWrapperArray*>(Self);
-		if (ArraySelf->ArrayProp && ArraySelf->ArrayInstance && !ArraySelf->OwnerContext.HasOwner())
-		{
-			Collector.AddReferencedObject(ArraySelf->ArrayProp);
-			FPyReferenceCollector::AddReferencedObjectsFromProperty(Collector, ArraySelf->ArrayProp, ArraySelf->ArrayInstance);
-		}
-	};
+		Collector.AddReferencedObject(Self->ArrayProp);
+		FPyReferenceCollector::AddReferencedObjectsFromProperty(Collector, Self->ArrayProp, Self->ArrayInstance);
+	}
 }
 
 #endif	// WITH_PYTHON
diff --git a/Engine/Plugins/Experimental/PythonScriptPlugin/Source/PythonScriptPlugin/Private/PyWrapperArray.h b/Engine/Plugins/Experimental/PythonScriptPlugin/Source/PythonScriptPlugin/Private/PyWrapperArray.h
--- a/Engine/Plugins/Experimental/PythonScriptPlugin/Source/PythonScriptPlugin/Private/PyWrapperArray.h
+++ b/Engine/Plugins/Experimental/PythonScriptPlugin/Source/PythonScriptPlugin/Private/PyWrapperArray.h
@@ -45,10 +45,10 @@ struct FPyWrapperArray : public FPyWrapperBase
 	static bool ValidateInternalState(FPyWrapperArray* InSelf);
 
 	/** Cast the given Python object to this wrapped type (returns a new reference) */
-	static FPyWrapperArray* CastPyObject(PyObject* InPyObject);
+	static FPyWrapperArray* CastPyObject(PyObject* InPyObject, FPyConversionResult* OutCastResult = nullptr);
 
 	/** Cast the given Python object to this wrapped type, or attempt to convert the type into a new wrapped instance (returns a new reference) */
-	static FPyWrapperArray* CastPyObject(PyObject* InPyObject, PyTypeObject* InType, const PyUtil::FPropertyDef& InElementDef);
+	static FPyWrapperArray* CastPyObject(PyObject* InPyObject, PyTypeObject* InType, const PyUtil::FPropertyDef& InElementDef, FPyConversionResult* OutCastResult = nullptr);
 
 	/** Get the length of this container (equivalent to 'len(x)' in Python) */
 	static Py_ssize_t Len(FPyWrapperArray* InSelf);
@@ -109,9 +109,14 @@ struct FPyWrapperArray : public FPyWrapperBase
 /** Meta-data for all UE4 exposed array types */
 struct FPyWrapperArrayMetaData : public FPyWrapperBaseMetaData
 {
-	PY_OVERRIDE_GETSET_METADATA(FPyWrapperArrayMetaData)
+	PY_METADATA_METHODS(FPyWrapperArrayMetaData, FGuid(0xFCD8DF06, 0x7AD6426F, 0x97482D82, 0x6E100C39))
 
-	FPyWrapperArrayMetaData();
+	FPyWrapperArrayMetaData()
+	{
+	}
+
+	/** Add object references from the given Python object to the given collector */
+	virtual void AddReferencedObjects(FPyWrapperBase* Instance, FReferenceCollector& Collector) override;
 };
 
 typedef TPyPtr<FPyWrapperArray> FPyWrapperArrayPtr;
diff --git a/Engine/Plugins/Experimental/PythonScriptPlugin/Source/PythonScriptPlugin/Private/PyWrapperBase.h b/Engine/Plugins/Experimental/PythonScriptPlugin/Source/PythonScriptPlugin/Private/PyWrapperBase.h
--- a/Engine/Plugins/Experimental/PythonScriptPlugin/Source/PythonScriptPlugin/Private/PyWrapperBase.h
+++ b/Engine/Plugins/Experimental/PythonScriptPlugin/Source/PythonScriptPlugin/Private/PyWrapperBase.h
@@ -5,6 +5,8 @@
 #include "IncludePython.h"
 #include "PyPtr.h"
 #include "PyConversionMethod.h"
+#include "PyConversionResult.h"
+#include "Misc/Guid.h"
 #include "UObject/GCObject.h"
 
 #if WITH_PYTHON
@@ -39,6 +41,11 @@ struct FPyWrapperBase
 	static TYPE* GetMetaData(PyTypeObject* PyType) { return (TYPE*)FPyWrapperBaseMetaData::GetMetaData(PyType); }				\
 	static TYPE* GetMetaData(FPyWrapperBase* Instance) { return (TYPE*)FPyWrapperBaseMetaData::GetMetaData(Instance); }
 
+#define PY_METADATA_METHODS(TYPE, GUID)																							\
+	PY_OVERRIDE_GETSET_METADATA(TYPE)																							\
+	static FGuid StaticTypeId() { return (GUID); }																				\
+	virtual FGuid GetTypeId() const override { return StaticTypeId(); }
+
 /** Base meta-data for all UE4 exposed types */
 struct FPyWrapperBaseMetaData
 {
@@ -52,13 +59,20 @@ struct FPyWrapperBaseMetaData
 	static FPyWrapperBaseMetaData* GetMetaData(FPyWrapperBase* Instance);
 
 	FPyWrapperBaseMetaData()
-		: AddReferencedObjects(nullptr)
 	{
 	}
 
-	/** AddReferencedObjects */
-	typedef void(*FAROFunc)(FPyWrapperBase*, FReferenceCollector&);
-	FAROFunc AddReferencedObjects;
+	virtual ~FPyWrapperBaseMetaData()
+	{
+	}
+
+	/** Get the ID associated with this meta-data type */
+	virtual FGuid GetTypeId() const = 0;
+
+	/** Add object references from the given Python object to the given collector */
+	virtual void AddReferencedObjects(FPyWrapperBase* Instance, FReferenceCollector& Collector)
+	{
+	}
 };
 
 typedef TPyPtr<FPyWrapperBase> FPyWrapperBasePtr;
diff --git a/Engine/Plugins/Experimental/PythonScriptPlugin/Source/PythonScriptPlugin/Private/PyWrapperDelegate.cpp b/Engine/Plugins/Experimental/PythonScriptPlugin/Source/PythonScriptPlugin/Private/PyWrapperDelegate.cpp
--- a/Engine/Plugins/Experimental/PythonScriptPlugin/Source/PythonScriptPlugin/Private/PyWrapperDelegate.cpp
+++ b/Engine/Plugins/Experimental/PythonScriptPlugin/Source/PythonScriptPlugin/Private/PyWrapperDelegate.cpp
@@ -124,7 +124,6 @@ struct TPyWrapperDelegateImpl
 		if (Self)
 		{
 			new(&Self->OwnerContext) FPyWrapperOwnerContext();
-			Self->DelegateSignature = nullptr;
 			Self->DelegateInstance = nullptr;
 			new(&Self->InternalDelegateInstance) DelegateType();
 		}
@@ -150,21 +149,13 @@ struct TPyWrapperDelegateImpl
 			return BaseInit;
 		}
 
-		const UFunction* DelegateSignature = MetaDataType::GetDelegateSignature(InSelf);
-		if (!DelegateSignature)
-		{
-			PyUtil::SetPythonError(PyExc_Exception, InSelf, TEXT("DelegateSignature is null"));
-			return -1;
-		}
-
-		InSelf->DelegateSignature = DelegateSignature;
 		InSelf->DelegateInstance = &InSelf->InternalDelegateInstance;
 
 		FactoryType::Get().MapInstance(InSelf->DelegateInstance, InSelf);
 		return 0;
 	}
 
-	static int Init(WrapperType* InSelf, const FPyWrapperOwnerContext& InOwnerContext, const UFunction* InDelegateSignature, DelegateType* InValue, const EPyConversionMethod InConversionMethod)
+	static int Init(WrapperType* InSelf, const FPyWrapperOwnerContext& InOwnerContext, DelegateType* InValue, const EPyConversionMethod InConversionMethod)
 	{
 		InOwnerContext.AssertValidConversionMethod(InConversionMethod);
 
@@ -198,7 +189,6 @@ struct TPyWrapperDelegateImpl
 		check(DelegateInstanceToUse);
 
 		InSelf->OwnerContext = InOwnerContext;
-		InSelf->DelegateSignature = InDelegateSignature;
 		InSelf->DelegateInstance = DelegateInstanceToUse;
 
 		FactoryType::Get().MapInstance(InSelf->DelegateInstance, InSelf);
@@ -217,19 +207,12 @@ struct TPyWrapperDelegateImpl
 			InSelf->OwnerContext.Reset();
 		}
 
-		InSelf->DelegateSignature = nullptr;
 		InSelf->DelegateInstance = nullptr;
 		InSelf->InternalDelegateInstance.Clear();
 	}
 
 	static bool ValidateInternalState(WrapperType* InSelf)
 	{
-		if (!InSelf->DelegateSignature)
-		{
-			PyUtil::SetPythonError(PyExc_Exception, Py_TYPE(InSelf), TEXT("Internal Error - DelegateSignature is null!"));
-			return false;
-		}
-
 		if (!InSelf->DelegateInstance)
 		{
 			PyUtil::SetPythonError(PyExc_Exception, Py_TYPE(InSelf), TEXT("Internal Error - DelegateInstance is null!"));
@@ -254,7 +237,9 @@ struct TPyWrapperDelegateImpl
 			return nullptr;
 		}
 
-		if (InSelf->DelegateSignature->Children == nullptr)
+		const PyGenUtil::FGeneratedWrappedFunction& DelegateSignature = MetaDataType::GetDelegateSignature(InSelf);
+
+		if (DelegateSignature.Func->Children == nullptr)
 		{
 			// Simple case, no parameters or return value
 			FDelegateInvocation::Call(*InSelf->DelegateInstance, nullptr);
@@ -262,60 +247,30 @@ struct TPyWrapperDelegateImpl
 		}
 
 		// Complex case, parameters or return value
-		FStructOnScope DelegateParams(InSelf->DelegateSignature);
-		TArray<const UProperty*, TInlineAllocator<4>> OutParams;
-
-		// Add the return property first
-		if (const UProperty* ReturnProp = InSelf->DelegateSignature->GetReturnProperty())
+		TArray<PyObject*> Params;
+		if (!PyGenUtil::ParseMethodParameters(InArgs, nullptr, DelegateSignature.InputParams, "delegate", Params))
 		{
-			OutParams.Add(ReturnProp);
+			return nullptr;
 		}
 
-		// Add any other output params in order, and set the value of the input params from the Python args
-		int32 ArgIndex = 0;
-		const int32 NumArgs = PyTuple_GET_SIZE(InArgs);
-		for (TFieldIterator<const UProperty> ParamIt(InSelf->DelegateSignature); ParamIt; ++ParamIt)
+		FStructOnScope DelegateParams(DelegateSignature.Func);
+		PyGenUtil::ApplyParamDefaults(DelegateParams.GetStructMemory(), DelegateSignature.InputParams);
+		for (int32 ParamIndex = 0; ParamIndex < Params.Num(); ++ParamIndex)
 		{
-			const UProperty* Param = *ParamIt;
+			const PyGenUtil::FGeneratedWrappedMethodParameter& ParamDef = DelegateSignature.InputParams[ParamIndex];
 
-			if (PyUtil::IsInputParameter(Param))
+			PyObject* PyValue = Params[ParamIndex];
+			if (PyValue)
 			{
-				PyObject* ParsedArg = nullptr;
-				if (ArgIndex < NumArgs)
+				if (!PyConversion::NativizeProperty_InContainer(PyValue, ParamDef.ParamProp, DelegateParams.GetStructMemory(), 0))
 				{
-					ParsedArg = PyTuple_GET_ITEM(InArgs, ArgIndex);
-				}
-
-				if (!ParsedArg)
-				{
-					PyUtil::SetPythonError(PyExc_TypeError, InSelf, *FString::Printf(TEXT("Required argument at pos %d not found"), ArgIndex + 1));
-					return nullptr;
-				}
-
-				if (!PyConversion::NativizeProperty_InContainer(ParsedArg, Param, DelegateParams.GetStructMemory(), 0))
-				{
-					PyUtil::SetPythonError(PyExc_TypeError, InSelf, *FString::Printf(TEXT("Failed to convert argument at pos '%d' when calling delegate"), ArgIndex + 1));
+					PyUtil::SetPythonError(PyExc_TypeError, InSelf, *FString::Printf(TEXT("Failed to convert parameter '%s' when calling delegate"), UTF8_TO_TCHAR(ParamDef.ParamName.GetData())));
 					return nullptr;
 				}
-
-				++ArgIndex;
 			}
-
-			if (PyUtil::IsOutputParameter(Param))
-			{
-				OutParams.Add(Param);
-			}
-		}
-
-		// Too many arguments?
-		if (ArgIndex > NumArgs)
-		{
-			PyUtil::SetPythonError(PyExc_TypeError, InSelf, *FString::Printf(TEXT("Delegate takes at most %d argument%s (%d given)"), ArgIndex, (ArgIndex == 1 ? TEXT("") : TEXT("s")), NumArgs));
-			return nullptr;
 		}
-
 		FDelegateInvocation::Call(*InSelf->DelegateInstance, DelegateParams.GetStructMemory());
-		return PyUtil::PackReturnValues(InSelf->DelegateSignature, OutParams, DelegateParams.GetStructMemory(), *PyUtil::GetErrorContext(InSelf), TEXT("delegate"));
+		return PyGenUtil::PackReturnValues(DelegateParams.GetStructMemory(), DelegateSignature.OutputParams, *PyUtil::GetErrorContext(InSelf), TEXT("delegate"));
 	}
 };
 
@@ -338,9 +293,9 @@ int FPyWrapperDelegate::Init(FPyWrapperDelegate* InSelf)
 	return FPyWrapperDelegateImpl::Init(InSelf);
 }
 
-int FPyWrapperDelegate::Init(FPyWrapperDelegate* InSelf, const FPyWrapperOwnerContext& InOwnerContext, const UFunction* InDelegateSignature, FScriptDelegate* InValue, const EPyConversionMethod InConversionMethod)
+int FPyWrapperDelegate::Init(FPyWrapperDelegate* InSelf, const FPyWrapperOwnerContext& InOwnerContext, FScriptDelegate* InValue, const EPyConversionMethod InConversionMethod)
 {
-	return FPyWrapperDelegateImpl::Init(InSelf, InOwnerContext, InDelegateSignature, InValue, InConversionMethod);
+	return FPyWrapperDelegateImpl::Init(InSelf, InOwnerContext, InValue, InConversionMethod);
 }
 
 void FPyWrapperDelegate::Deinit(FPyWrapperDelegate* InSelf)
@@ -353,10 +308,14 @@ bool FPyWrapperDelegate::ValidateInternalState(FPyWrapperDelegate* InSelf)
 	return FPyWrapperDelegateImpl::ValidateInternalState(InSelf);
 }
 
-FPyWrapperDelegate* FPyWrapperDelegate::CastPyObject(PyObject* InPyObject)
+FPyWrapperDelegate* FPyWrapperDelegate::CastPyObject(PyObject* InPyObject, FPyConversionResult* OutCastResult)
 {
+	SetOptionalPyConversionResult(FPyConversionResult::Failure(), OutCastResult);
+
 	if (PyObject_IsInstance(InPyObject, (PyObject*)&PyWrapperDelegateType) == 1)
 	{
+		SetOptionalPyConversionResult(FPyConversionResult::Success(), OutCastResult);
+
 		Py_INCREF(InPyObject);
 		return (FPyWrapperDelegate*)InPyObject;
 	}
@@ -364,10 +323,14 @@ FPyWrapperDelegate* FPyWrapperDelegate::CastPyObject(PyObject* InPyObject)
 	return nullptr;
 }
 
-FPyWrapperDelegate* FPyWrapperDelegate::CastPyObject(PyObject* InPyObject, PyTypeObject* InType)
+FPyWrapperDelegate* FPyWrapperDelegate::CastPyObject(PyObject* InPyObject, PyTypeObject* InType, FPyConversionResult* OutCastResult)
 {
+	SetOptionalPyConversionResult(FPyConversionResult::Failure(), OutCastResult);
+
 	if (PyObject_IsInstance(InPyObject, (PyObject*)InType) == 1 && (InType == &PyWrapperDelegateType || PyObject_IsInstance(InPyObject, (PyObject*)&PyWrapperDelegateType) == 1))
 	{
+		SetOptionalPyConversionResult(Py_TYPE(InPyObject) == InType ? FPyConversionResult::Success() : FPyConversionResult::SuccessWithCoercion(), OutCastResult);
+
 		Py_INCREF(InPyObject);
 		return (FPyWrapperDelegate*)InPyObject;
 	}
@@ -396,9 +359,9 @@ int FPyWrapperMulticastDelegate::Init(FPyWrapperMulticastDelegate* InSelf)
 	return FPyWrapperMulticastDelegateImpl::Init(InSelf);
 }
 
-int FPyWrapperMulticastDelegate::Init(FPyWrapperMulticastDelegate* InSelf, const FPyWrapperOwnerContext& InOwnerContext, const UFunction* InDelegateSignature, FMulticastScriptDelegate* InValue, const EPyConversionMethod InConversionMethod)
+int FPyWrapperMulticastDelegate::Init(FPyWrapperMulticastDelegate* InSelf, const FPyWrapperOwnerContext& InOwnerContext, FMulticastScriptDelegate* InValue, const EPyConversionMethod InConversionMethod)
 {
-	return FPyWrapperMulticastDelegateImpl::Init(InSelf, InOwnerContext, InDelegateSignature, InValue, InConversionMethod);
+	return FPyWrapperMulticastDelegateImpl::Init(InSelf, InOwnerContext, InValue, InConversionMethod);
 }
 
 void FPyWrapperMulticastDelegate::Deinit(FPyWrapperMulticastDelegate* InSelf)
@@ -411,10 +374,14 @@ bool FPyWrapperMulticastDelegate::ValidateInternalState(FPyWrapperMulticastDeleg
 	return FPyWrapperMulticastDelegateImpl::ValidateInternalState(InSelf);
 }
 
-FPyWrapperMulticastDelegate* FPyWrapperMulticastDelegate::CastPyObject(PyObject* InPyObject)
+FPyWrapperMulticastDelegate* FPyWrapperMulticastDelegate::CastPyObject(PyObject* InPyObject, FPyConversionResult* OutCastResult)
 {
+	SetOptionalPyConversionResult(FPyConversionResult::Failure(), OutCastResult);
+
 	if (PyObject_IsInstance(InPyObject, (PyObject*)&PyWrapperMulticastDelegateType) == 1)
 	{
+		SetOptionalPyConversionResult(FPyConversionResult::Success(), OutCastResult);
+
 		Py_INCREF(InPyObject);
 		return (FPyWrapperMulticastDelegate*)InPyObject;
 	}
@@ -422,10 +389,14 @@ FPyWrapperMulticastDelegate* FPyWrapperMulticastDelegate::CastPyObject(PyObject*
 	return nullptr;
 }
 
-FPyWrapperMulticastDelegate* FPyWrapperMulticastDelegate::CastPyObject(PyObject* InPyObject, PyTypeObject* InType)
+FPyWrapperMulticastDelegate* FPyWrapperMulticastDelegate::CastPyObject(PyObject* InPyObject, PyTypeObject* InType, FPyConversionResult* OutCastResult)
 {
+	SetOptionalPyConversionResult(FPyConversionResult::Failure(), OutCastResult);
+
 	if (PyObject_IsInstance(InPyObject, (PyObject*)InType) == 1 && (InType == &PyWrapperMulticastDelegateType || PyObject_IsInstance(InPyObject, (PyObject*)&PyWrapperMulticastDelegateType) == 1))
 	{
+		SetOptionalPyConversionResult(Py_TYPE(InPyObject) == InType ? FPyConversionResult::Success() : FPyConversionResult::SuccessWithCoercion(), OutCastResult);
+
 		Py_INCREF(InPyObject);
 		return (FPyWrapperMulticastDelegate*)InPyObject;
 	}
@@ -461,7 +432,8 @@ PyTypeObject InitializePyWrapperDelegateType()
 				return BaseInit;
 			}
 
-			if (PyTuple_Size(InArgs) > 0 && !PyDelegateUtil::PythonArgsToDelegate(InArgs, InSelf->DelegateSignature, *InSelf->DelegateInstance, TEXT("call"), *PyUtil::GetErrorContext(InSelf)))
+			const PyGenUtil::FGeneratedWrappedFunction& DelegateSignature = FPyWrapperDelegateMetaData::GetDelegateSignature(InSelf);
+			if (PyTuple_Size(InArgs) > 0 && !PyDelegateUtil::PythonArgsToDelegate(InArgs, DelegateSignature.Func, *InSelf->DelegateInstance, TEXT("call"), *PyUtil::GetErrorContext(InSelf)))
 			{
 				return -1;
 			}
@@ -529,7 +501,8 @@ PyTypeObject InitializePyWrapperDelegateType()
 				return nullptr;
 			}
 
-			return (PyObject*)FPyWrapperDelegateFactory::Get().CreateInstance(InSelf->DelegateSignature, InSelf->DelegateInstance, FPyWrapperOwnerContext(), EPyConversionMethod::Copy);
+			const PyGenUtil::FGeneratedWrappedFunction& DelegateSignature = FPyWrapperDelegateMetaData::GetDelegateSignature(InSelf);
+			return (PyObject*)FPyWrapperDelegateFactory::Get().CreateInstance(DelegateSignature.Func, InSelf->DelegateInstance, FPyWrapperOwnerContext(), EPyConversionMethod::Copy);
 		}
 
 		static PyObject* IsBound(FPyWrapperDelegate* InSelf)
@@ -554,7 +527,8 @@ PyTypeObject InitializePyWrapperDelegateType()
 				return nullptr;
 			}
 
-			if (!PyDelegateUtil::PythonArgsToDelegate(InArgs, InSelf->DelegateSignature, *InSelf->DelegateInstance, TEXT("bind_function"), *PyUtil::GetErrorContext(InSelf)))
+			const PyGenUtil::FGeneratedWrappedFunction& DelegateSignature = FPyWrapperDelegateMetaData::GetDelegateSignature(InSelf);
+			if (!PyDelegateUtil::PythonArgsToDelegate(InArgs, DelegateSignature.Func, *InSelf->DelegateInstance, TEXT("bind_function"), *PyUtil::GetErrorContext(InSelf)))
 			{
 				return nullptr;
 			}
@@ -661,8 +635,10 @@ PyTypeObject InitializePyWrapperMulticastDelegateType()
 
 			if (PyTuple_Size(InArgs) > 0)
 			{
+				const PyGenUtil::FGeneratedWrappedFunction& DelegateSignature = FPyWrapperMulticastDelegateMetaData::GetDelegateSignature(InSelf);
+
 				FScriptDelegate Delegate;
-				if (!PyDelegateUtil::PythonArgsToDelegate(InArgs, InSelf->DelegateSignature, Delegate, TEXT("call"), *PyUtil::GetErrorContext(InSelf)))
+				if (!PyDelegateUtil::PythonArgsToDelegate(InArgs, DelegateSignature.Func, Delegate, TEXT("call"), *PyUtil::GetErrorContext(InSelf)))
 				{
 					return -1;
 				}
@@ -732,7 +708,8 @@ PyTypeObject InitializePyWrapperMulticastDelegateType()
 				return nullptr;
 			}
 
-			return (PyObject*)FPyWrapperMulticastDelegateFactory::Get().CreateInstance(InSelf->DelegateSignature, InSelf->DelegateInstance, FPyWrapperOwnerContext(), EPyConversionMethod::Copy);
+			const PyGenUtil::FGeneratedWrappedFunction& DelegateSignature = FPyWrapperMulticastDelegateMetaData::GetDelegateSignature(InSelf);
+			return (PyObject*)FPyWrapperMulticastDelegateFactory::Get().CreateInstance(DelegateSignature.Func, InSelf->DelegateInstance, FPyWrapperOwnerContext(), EPyConversionMethod::Copy);
 		}
 
 		static PyObject* IsBound(FPyWrapperMulticastDelegate* InSelf)
@@ -757,8 +734,10 @@ PyTypeObject InitializePyWrapperMulticastDelegateType()
 				return nullptr;
 			}
 
+			const PyGenUtil::FGeneratedWrappedFunction& DelegateSignature = FPyWrapperMulticastDelegateMetaData::GetDelegateSignature(InSelf);
+
 			FScriptDelegate Delegate;
-			if (!PyDelegateUtil::PythonArgsToDelegate(InArgs, InSelf->DelegateSignature, Delegate, TEXT("add_function"), *PyUtil::GetErrorContext(InSelf)))
+			if (!PyDelegateUtil::PythonArgsToDelegate(InArgs, DelegateSignature.Func, Delegate, TEXT("add_function"), *PyUtil::GetErrorContext(InSelf)))
 			{
 				return nullptr;
 			}
@@ -775,8 +754,10 @@ PyTypeObject InitializePyWrapperMulticastDelegateType()
 				return nullptr;
 			}
 
+			const PyGenUtil::FGeneratedWrappedFunction& DelegateSignature = FPyWrapperMulticastDelegateMetaData::GetDelegateSignature(InSelf);
+
 			FScriptDelegate Delegate;
-			if (!PyDelegateUtil::PythonArgsToDelegate(InArgs, InSelf->DelegateSignature, Delegate, TEXT("add_function_unique"), *PyUtil::GetErrorContext(InSelf)))
+			if (!PyDelegateUtil::PythonArgsToDelegate(InArgs, DelegateSignature.Func, Delegate, TEXT("add_function_unique"), *PyUtil::GetErrorContext(InSelf)))
 			{
 				return nullptr;
 			}
@@ -793,8 +774,10 @@ PyTypeObject InitializePyWrapperMulticastDelegateType()
 				return nullptr;
 			}
 
+			const PyGenUtil::FGeneratedWrappedFunction& DelegateSignature = FPyWrapperMulticastDelegateMetaData::GetDelegateSignature(InSelf);
+
 			FScriptDelegate Delegate;
-			if (!PyDelegateUtil::PythonArgsToDelegate(InArgs, InSelf->DelegateSignature, Delegate, TEXT("remove_function"), *PyUtil::GetErrorContext(InSelf)))
+			if (!PyDelegateUtil::PythonArgsToDelegate(InArgs, DelegateSignature.Func, Delegate, TEXT("remove_function"), *PyUtil::GetErrorContext(InSelf)))
 			{
 				return nullptr;
 			}
@@ -835,8 +818,10 @@ PyTypeObject InitializePyWrapperMulticastDelegateType()
 				return nullptr;
 			}
 
+			const PyGenUtil::FGeneratedWrappedFunction& DelegateSignature = FPyWrapperMulticastDelegateMetaData::GetDelegateSignature(InSelf);
+
 			FScriptDelegate Delegate;
-			if (!PyDelegateUtil::PythonArgsToDelegate(InArgs, InSelf->DelegateSignature, Delegate, TEXT("contains_function"), *PyUtil::GetErrorContext(InSelf)))
+			if (!PyDelegateUtil::PythonArgsToDelegate(InArgs, DelegateSignature.Func, Delegate, TEXT("contains_function"), *PyUtil::GetErrorContext(InSelf)))
 			{
 				return nullptr;
 			}
diff --git a/Engine/Plugins/Experimental/PythonScriptPlugin/Source/PythonScriptPlugin/Private/PyWrapperDelegate.h b/Engine/Plugins/Experimental/PythonScriptPlugin/Source/PythonScriptPlugin/Private/PyWrapperDelegate.h
--- a/Engine/Plugins/Experimental/PythonScriptPlugin/Source/PythonScriptPlugin/Private/PyWrapperDelegate.h
+++ b/Engine/Plugins/Experimental/PythonScriptPlugin/Source/PythonScriptPlugin/Private/PyWrapperDelegate.h
@@ -3,6 +3,7 @@
 #pragma once
 
 #include "PyWrapperBase.h"
+#include "PyGenUtil.h"
 #include "PyWrapperOwnerContext.h"
 #include "UObject/WeakObjectPtr.h"
 
@@ -24,9 +25,6 @@ struct TPyWrapperDelegate : public FPyWrapperBase
 	/** The owner of the wrapped delegate instance (if any) */
 	FPyWrapperOwnerContext OwnerContext;
 
-	/** Function representing the signature of this delegate instance */
-	const UFunction* DelegateSignature;
-
 	/** Wrapped delegate instance */
 	DelegateType* DelegateInstance;
 
@@ -41,25 +39,25 @@ struct TPyWrapperDelegateMetaData : public FPyWrapperBaseMetaData
 	PY_OVERRIDE_GETSET_METADATA(TPyWrapperDelegateMetaData)
 
 	TPyWrapperDelegateMetaData()
-		: DelegateSignature(nullptr)
 	{
 	}
 
 	/** Get the delegate signature from the given type */
-	static const UFunction* GetDelegateSignature(PyTypeObject* PyType)
+	static const PyGenUtil::FGeneratedWrappedFunction& GetDelegateSignature(PyTypeObject* PyType)
 	{
 		TPyWrapperDelegateMetaData* PyWrapperMetaData = TPyWrapperDelegateMetaData::GetMetaData(PyType);
-		return PyWrapperMetaData ? PyWrapperMetaData->DelegateSignature : nullptr;
+		static const PyGenUtil::FGeneratedWrappedFunction NullDelegateSignature = PyGenUtil::FGeneratedWrappedFunction();
+		return PyWrapperMetaData ? PyWrapperMetaData->DelegateSignature : NullDelegateSignature;
 	}
 
 	/** Get the delegate signature from the type of the given instance */
-	static const UFunction* GetDelegateSignature(WrapperType* Instance)
+	static const PyGenUtil::FGeneratedWrappedFunction& GetDelegateSignature(WrapperType* Instance)
 	{
 		return GetDelegateSignature(Py_TYPE(Instance));
 	}
 
 	/** Unreal function representing the signature for the delegate */
-	const UFunction* DelegateSignature;
+	PyGenUtil::FGeneratedWrappedFunction DelegateSignature;
 };
 
 /** Type for all UE4 exposed delegate instances */
@@ -75,7 +73,7 @@ struct FPyWrapperDelegate : public TPyWrapperDelegate<FScriptDelegate>
 	static int Init(FPyWrapperDelegate* InSelf);
 
 	/** Initialize this wrapper instance to the given value (called via tp_init for Python, or directly in C++) */
-	static int Init(FPyWrapperDelegate* InSelf, const FPyWrapperOwnerContext& InOwnerContext, const UFunction* InDelegateSignature, FScriptDelegate* InValue, const EPyConversionMethod InConversionMethod);
+	static int Init(FPyWrapperDelegate* InSelf, const FPyWrapperOwnerContext& InOwnerContext, FScriptDelegate* InValue, const EPyConversionMethod InConversionMethod);
 
 	/** Deinitialize this wrapper instance (called via Init and Free to restore the instance to its New state) */
 	static void Deinit(FPyWrapperDelegate* InSelf);
@@ -84,10 +82,10 @@ struct FPyWrapperDelegate : public TPyWrapperDelegate<FScriptDelegate>
 	static bool ValidateInternalState(FPyWrapperDelegate* InSelf);
 
 	/** Cast the given Python object to this wrapped type (returns a new reference) */
-	static FPyWrapperDelegate* CastPyObject(PyObject* InPyObject);
+	static FPyWrapperDelegate* CastPyObject(PyObject* InPyObject, FPyConversionResult* OutCastResult = nullptr);
 
 	/** Cast the given Python object to this wrapped type, or attempt to convert the type into a new wrapped instance (returns a new reference) */
-	static FPyWrapperDelegate* CastPyObject(PyObject* InPyObject, PyTypeObject* InType);
+	static FPyWrapperDelegate* CastPyObject(PyObject* InPyObject, PyTypeObject* InType, FPyConversionResult* OutCastResult = nullptr);
 
 	/** Call the delegate */
 	static PyObject* CallDelegate(FPyWrapperDelegate* InSelf, PyObject* InArgs);
@@ -96,7 +94,7 @@ struct FPyWrapperDelegate : public TPyWrapperDelegate<FScriptDelegate>
 /** Meta-data for all UE4 exposed delegate types */
 struct FPyWrapperDelegateMetaData : public TPyWrapperDelegateMetaData<FPyWrapperDelegate>
 {
-	PY_OVERRIDE_GETSET_METADATA(FPyWrapperDelegateMetaData)
+	PY_METADATA_METHODS(FPyWrapperDelegateMetaData, FGuid(0xCB3D0485, 0x8A3A443E, 0xBEE336F4, 0x82888A81))
 };
 
 /** Type for all UE4 exposed multicast delegate instances */
@@ -112,7 +110,7 @@ struct FPyWrapperMulticastDelegate : public TPyWrapperDelegate<FMulticastScriptD
 	static int Init(FPyWrapperMulticastDelegate* InSelf);
 
 	/** Initialize this wrapper instance to the given value (called via tp_init for Python, or directly in C++) */
-	static int Init(FPyWrapperMulticastDelegate* InSelf, const FPyWrapperOwnerContext& InOwnerContext, const UFunction* InDelegateSignature, FMulticastScriptDelegate* InValue, const EPyConversionMethod InConversionMethod);
+	static int Init(FPyWrapperMulticastDelegate* InSelf, const FPyWrapperOwnerContext& InOwnerContext, FMulticastScriptDelegate* InValue, const EPyConversionMethod InConversionMethod);
 
 	/** Deinitialize this wrapper instance (called via Init and Free to restore the instance to its New state) */
 	static void Deinit(FPyWrapperMulticastDelegate* InSelf);
@@ -121,10 +119,10 @@ struct FPyWrapperMulticastDelegate : public TPyWrapperDelegate<FMulticastScriptD
 	static bool ValidateInternalState(FPyWrapperMulticastDelegate* InSelf);
 
 	/** Cast the given Python object to this wrapped type (returns a new reference) */
-	static FPyWrapperMulticastDelegate* CastPyObject(PyObject* InPyObject);
+	static FPyWrapperMulticastDelegate* CastPyObject(PyObject* InPyObject, FPyConversionResult* OutCastResult = nullptr);
 
 	/** Cast the given Python object to this wrapped type, or attempt to convert the type into a new wrapped instance (returns a new reference) */
-	static FPyWrapperMulticastDelegate* CastPyObject(PyObject* InPyObject, PyTypeObject* InType);
+	static FPyWrapperMulticastDelegate* CastPyObject(PyObject* InPyObject, PyTypeObject* InType, FPyConversionResult* OutCastResult = nullptr);
 
 	/** Call the delegate */
 	static PyObject* CallDelegate(FPyWrapperMulticastDelegate* InSelf, PyObject* InArgs);
@@ -133,7 +131,7 @@ struct FPyWrapperMulticastDelegate : public TPyWrapperDelegate<FMulticastScriptD
 /** Meta-data for all UE4 exposed multicast delegate types */
 struct FPyWrapperMulticastDelegateMetaData : public TPyWrapperDelegateMetaData<FPyWrapperMulticastDelegate>
 {
-	PY_OVERRIDE_GETSET_METADATA(FPyWrapperMulticastDelegateMetaData)
+	PY_METADATA_METHODS(FPyWrapperMulticastDelegateMetaData, FGuid(0x448FB4DA, 0x38DC4386, 0xBCAFF448, 0x29C0F3A4))
 };
 
 typedef TPyPtr<FPyWrapperDelegate> FPyWrapperDelegatePtr;
diff --git a/Engine/Plugins/Experimental/PythonScriptPlugin/Source/PythonScriptPlugin/Private/PyWrapperEnum.cpp b/Engine/Plugins/Experimental/PythonScriptPlugin/Source/PythonScriptPlugin/Private/PyWrapperEnum.cpp
--- a/Engine/Plugins/Experimental/PythonScriptPlugin/Source/PythonScriptPlugin/Private/PyWrapperEnum.cpp
+++ b/Engine/Plugins/Experimental/PythonScriptPlugin/Source/PythonScriptPlugin/Private/PyWrapperEnum.cpp
@@ -244,7 +244,7 @@ struct FPythonGeneratedEnumUtil
 		}
 
 		// Build the definition data for the new enum value
-		UPythonGeneratedEnum::FEnumValueDef& EnumValueDef = *InEnum->EnumValueDefs[InEnum->EnumValueDefs.Add(MakeShared<UPythonGeneratedEnum::FEnumValueDef>())];
+		UPythonGeneratedEnum::FEnumValueDef& EnumValueDef = *InEnum->EnumValueDefs.Add_GetRef(MakeShared<UPythonGeneratedEnum::FEnumValueDef>());
 		EnumValueDef.Value = EnumValue;
 		EnumValueDef.Name = InFieldName;
 
diff --git a/Engine/Plugins/Experimental/PythonScriptPlugin/Source/PythonScriptPlugin/Private/PyWrapperEnum.h b/Engine/Plugins/Experimental/PythonScriptPlugin/Source/PythonScriptPlugin/Private/PyWrapperEnum.h
--- a/Engine/Plugins/Experimental/PythonScriptPlugin/Source/PythonScriptPlugin/Private/PyWrapperEnum.h
+++ b/Engine/Plugins/Experimental/PythonScriptPlugin/Source/PythonScriptPlugin/Private/PyWrapperEnum.h
@@ -27,7 +27,7 @@ struct FPyWrapperEnum : public FPyWrapperBase
 /** Meta-data for all UE4 exposed enum types */
 struct FPyWrapperEnumMetaData : public FPyWrapperBaseMetaData
 {
-	PY_OVERRIDE_GETSET_METADATA(FPyWrapperEnumMetaData)
+	PY_METADATA_METHODS(FPyWrapperEnumMetaData, FGuid(0x1D69987C, 0x2F624403, 0x8379FCB5, 0xF896B595))
 
 	FPyWrapperEnumMetaData();
 
diff --git a/Engine/Plugins/Experimental/PythonScriptPlugin/Source/PythonScriptPlugin/Private/PyWrapperFixedArray.cpp b/Engine/Plugins/Experimental/PythonScriptPlugin/Source/PythonScriptPlugin/Private/PyWrapperFixedArray.cpp
--- a/Engine/Plugins/Experimental/PythonScriptPlugin/Source/PythonScriptPlugin/Private/PyWrapperFixedArray.cpp
+++ b/Engine/Plugins/Experimental/PythonScriptPlugin/Source/PythonScriptPlugin/Private/PyWrapperFixedArray.cpp
@@ -259,10 +259,14 @@ bool FPyWrapperFixedArray::ValidateInternalState(FPyWrapperFixedArray* InSelf)
 	return true;
 }
 
-FPyWrapperFixedArray* FPyWrapperFixedArray::CastPyObject(PyObject* InPyObject)
+FPyWrapperFixedArray* FPyWrapperFixedArray::CastPyObject(PyObject* InPyObject, FPyConversionResult* OutCastResult)
 {
+	SetOptionalPyConversionResult(FPyConversionResult::Failure(), OutCastResult);
+
 	if (PyObject_IsInstance(InPyObject, (PyObject*)&PyWrapperFixedArrayType) == 1)
 	{
+		SetOptionalPyConversionResult(FPyConversionResult::Success(), OutCastResult);
+
 		Py_INCREF(InPyObject);
 		return (FPyWrapperFixedArray*)InPyObject;
 	}
@@ -270,8 +274,10 @@ FPyWrapperFixedArray* FPyWrapperFixedArray::CastPyObject(PyObject* InPyObject)
 	return nullptr;
 }
 
-FPyWrapperFixedArray* FPyWrapperFixedArray::CastPyObject(PyObject* InPyObject, PyTypeObject* InType, const PyUtil::FPropertyDef& InPropDef)
+FPyWrapperFixedArray* FPyWrapperFixedArray::CastPyObject(PyObject* InPyObject, PyTypeObject* InType, const PyUtil::FPropertyDef& InPropDef, FPyConversionResult* OutCastResult)
 {
+	SetOptionalPyConversionResult(FPyConversionResult::Failure(), OutCastResult);
+
 	if (PyObject_IsInstance(InPyObject, (PyObject*)InType) == 1 && (InType == &PyWrapperFixedArrayType || PyObject_IsInstance(InPyObject, (PyObject*)&PyWrapperFixedArrayType) == 1))
 	{
 		FPyWrapperFixedArray* Self = (FPyWrapperFixedArray*)InPyObject;
@@ -283,6 +289,8 @@ FPyWrapperFixedArray* FPyWrapperFixedArray::CastPyObject(PyObject* InPyObject, P
 		const PyUtil::FPropertyDef SelfPropDef = Self->ArrayProp;
 		if (SelfPropDef == InPropDef)
 		{
+			SetOptionalPyConversionResult(FPyConversionResult::Success(), OutCastResult);
+
 			Py_INCREF(Self);
 			return Self;
 		}
@@ -311,6 +319,7 @@ FPyWrapperFixedArray* FPyWrapperFixedArray::CastPyObject(PyObject* InPyObject, P
 			}
 		}
 
+		SetOptionalPyConversionResult(FPyConversionResult::SuccessWithCoercion(), OutCastResult);
 		return NewArray.Release();
 	}
 
@@ -343,6 +352,7 @@ FPyWrapperFixedArray* FPyWrapperFixedArray::CastPyObject(PyObject* InPyObject, P
 					}
 				}
 
+				SetOptionalPyConversionResult(FPyConversionResult::SuccessWithCoercion(), OutCastResult);
 				return NewArray.Release();
 			}
 		}
@@ -889,17 +899,14 @@ PyTypeObject InitializePyWrapperFixedArrayIteratorType()
 PyTypeObject PyWrapperFixedArrayType = InitializePyWrapperFixedArrayType();
 PyTypeObject PyWrapperFixedArrayIteratorType = InitializePyWrapperFixedArrayIteratorType();
 
-FPyWrapperFixedArrayMetaData::FPyWrapperFixedArrayMetaData()
+void FPyWrapperFixedArrayMetaData::AddReferencedObjects(FPyWrapperBase* Instance, FReferenceCollector& Collector)
 {
-	AddReferencedObjects = [](FPyWrapperBase* Self, FReferenceCollector& Collector)
+	FPyWrapperFixedArray* Self = static_cast<FPyWrapperFixedArray*>(Instance);
+	if (Self->ArrayProp && Self->ArrayInstance && !Self->OwnerContext.HasOwner())
 	{
-		FPyWrapperFixedArray* FixedArraySelf = static_cast<FPyWrapperFixedArray*>(Self);
-		if (FixedArraySelf->ArrayProp && FixedArraySelf->ArrayInstance && !FixedArraySelf->OwnerContext.HasOwner())
-		{
-			Collector.AddReferencedObject(FixedArraySelf->ArrayProp);
-			FPyReferenceCollector::AddReferencedObjectsFromProperty(Collector, FixedArraySelf->ArrayProp, FixedArraySelf->ArrayInstance);
-		}
-	};
+		Collector.AddReferencedObject(Self->ArrayProp);
+		FPyReferenceCollector::AddReferencedObjectsFromProperty(Collector, Self->ArrayProp, Self->ArrayInstance);
+	}
 }
 
 #endif	// WITH_PYTHON
diff --git a/Engine/Plugins/Experimental/PythonScriptPlugin/Source/PythonScriptPlugin/Private/PyWrapperFixedArray.h b/Engine/Plugins/Experimental/PythonScriptPlugin/Source/PythonScriptPlugin/Private/PyWrapperFixedArray.h
--- a/Engine/Plugins/Experimental/PythonScriptPlugin/Source/PythonScriptPlugin/Private/PyWrapperFixedArray.h
+++ b/Engine/Plugins/Experimental/PythonScriptPlugin/Source/PythonScriptPlugin/Private/PyWrapperFixedArray.h
@@ -45,10 +45,10 @@ struct FPyWrapperFixedArray : public FPyWrapperBase
 	static bool ValidateInternalState(FPyWrapperFixedArray* InSelf);
 
 	/** Cast the given Python object to this wrapped type (returns a new reference) */
-	static FPyWrapperFixedArray* CastPyObject(PyObject* InPyObject);
+	static FPyWrapperFixedArray* CastPyObject(PyObject* InPyObject, FPyConversionResult* OutCastResult = nullptr);
 
 	/** Cast the given Python object to this wrapped type, or attempt to convert the type into a new wrapped instance (returns a new reference) */
-	static FPyWrapperFixedArray* CastPyObject(PyObject* InPyObject, PyTypeObject* InType, const PyUtil::FPropertyDef& InPropDef);
+	static FPyWrapperFixedArray* CastPyObject(PyObject* InPyObject, PyTypeObject* InType, const PyUtil::FPropertyDef& InPropDef, FPyConversionResult* OutCastResult = nullptr);
 
 	/** Get the raw pointer to the element at index N */
 	static void* GetItemPtr(FPyWrapperFixedArray* InSelf, Py_ssize_t InIndex);
@@ -75,9 +75,14 @@ struct FPyWrapperFixedArray : public FPyWrapperBase
 /** Meta-data for all UE4 exposed fixed-array types */
 struct FPyWrapperFixedArrayMetaData : public FPyWrapperBaseMetaData
 {
-	PY_OVERRIDE_GETSET_METADATA(FPyWrapperFixedArrayMetaData)
+	PY_METADATA_METHODS(FPyWrapperFixedArrayMetaData, FGuid(0x8A8CA89E, 0xC8F04627, 0xA7D2B568, 0xB086C02E))
 
-	FPyWrapperFixedArrayMetaData();
+	FPyWrapperFixedArrayMetaData()
+	{
+	}
+
+	/** Add object references from the given Python object to the given collector */
+	virtual void AddReferencedObjects(FPyWrapperBase* Instance, FReferenceCollector& Collector) override;
 };
 
 typedef TPyPtr<FPyWrapperFixedArray> FPyWrapperFixedArrayPtr;
