Copying //UE4/Dev-Enterprise to //UE4/Dev-Main (Source: //UE4/Dev-Enterprise @ 3972172)

diff --git a/Engine/Plugins/Experimental/PythonScriptPlugin/Source/PythonScriptPlugin/Private/PyWrapperMap.cpp b/Engine/Plugins/Experimental/PythonScriptPlugin/Source/PythonScriptPlugin/Private/PyWrapperMap.cpp
--- a/Engine/Plugins/Experimental/PythonScriptPlugin/Source/PythonScriptPlugin/Private/PyWrapperMap.cpp
+++ b/Engine/Plugins/Experimental/PythonScriptPlugin/Source/PythonScriptPlugin/Private/PyWrapperMap.cpp
@@ -609,10 +609,14 @@ bool FPyWrapperMap::ValidateInternalState(FPyWrapperMap* InSelf)
 	return true;
 }
 
-FPyWrapperMap* FPyWrapperMap::CastPyObject(PyObject* InPyObject)
+FPyWrapperMap* FPyWrapperMap::CastPyObject(PyObject* InPyObject, FPyConversionResult* OutCastResult)
 {
+	SetOptionalPyConversionResult(FPyConversionResult::Failure(), OutCastResult);
+
 	if (PyObject_IsInstance(InPyObject, (PyObject*)&PyWrapperMapType) == 1)
 	{
+		SetOptionalPyConversionResult(FPyConversionResult::Success(), OutCastResult);
+
 		Py_INCREF(InPyObject);
 		return (FPyWrapperMap*)InPyObject;
 	}
@@ -620,8 +624,10 @@ FPyWrapperMap* FPyWrapperMap::CastPyObject(PyObject* InPyObject)
 	return nullptr;
 }
 
-FPyWrapperMap* FPyWrapperMap::CastPyObject(PyObject* InPyObject, PyTypeObject* InType, const PyUtil::FPropertyDef& InKeyDef, const PyUtil::FPropertyDef& InValueDef)
+FPyWrapperMap* FPyWrapperMap::CastPyObject(PyObject* InPyObject, PyTypeObject* InType, const PyUtil::FPropertyDef& InKeyDef, const PyUtil::FPropertyDef& InValueDef, FPyConversionResult* OutCastResult)
 {
+	SetOptionalPyConversionResult(FPyConversionResult::Failure(), OutCastResult);
+
 	if (PyObject_IsInstance(InPyObject, (PyObject*)InType) == 1 && (InType == &PyWrapperMapType || PyObject_IsInstance(InPyObject, (PyObject*)&PyWrapperMapType) == 1))
 	{
 		FPyWrapperMap* Self = (FPyWrapperMap*)InPyObject;
@@ -635,6 +641,8 @@ FPyWrapperMap* FPyWrapperMap::CastPyObject(PyObject* InPyObject, PyTypeObject* I
 		const PyUtil::FPropertyDef SelfValuePropDef = Self->MapProp->ValueProp;
 		if (SelfKeyPropDef == InKeyDef && SelfValuePropDef == InValueDef)
 		{
+			SetOptionalPyConversionResult(FPyConversionResult::Success(), OutCastResult);
+
 			Py_INCREF(Self);
 			return Self;
 		}
@@ -693,6 +701,7 @@ FPyWrapperMap* FPyWrapperMap::CastPyObject(PyObject* InPyObject, PyTypeObject* I
 			NewScriptMapHelper.Rehash();
 		}
 
+		SetOptionalPyConversionResult(FPyConversionResult::SuccessWithCoercion(), OutCastResult);
 		return NewMap.Release();
 	}
 
@@ -779,6 +788,8 @@ FPyWrapperMap* FPyWrapperMap::CastPyObject(PyObject* InPyObject, PyTypeObject* I
 				}
 
 				NewScriptMapHelper.Rehash();
+
+				SetOptionalPyConversionResult(FPyConversionResult::SuccessWithCoercion(), OutCastResult);
 				return NewMap.Release();
 			}
 		}
@@ -1830,17 +1841,14 @@ PyTypeObject PyWrapperMapItemViewType = InitializePyWrapperMapSetViewType<FPyWra
 PyTypeObject PyWrapperMapKeyViewType = InitializePyWrapperMapSetViewType<FPyWrapperMapKeyView>("MapKeyView");
 PyTypeObject PyWrapperMapValueViewType = InitializePyWrapperMapViewType<FPyWrapperMapValueView>("MapValueView");
 
-FPyWrapperMapMetaData::FPyWrapperMapMetaData()
+void FPyWrapperMapMetaData::AddReferencedObjects(FPyWrapperBase* Instance, FReferenceCollector& Collector)
 {
-	AddReferencedObjects = [](FPyWrapperBase* Self, FReferenceCollector& Collector)
+	FPyWrapperMap* Self = static_cast<FPyWrapperMap*>(Instance);
+	if (Self->MapProp && Self->MapInstance && !Self->OwnerContext.HasOwner())
 	{
-		FPyWrapperMap* MapSelf = static_cast<FPyWrapperMap*>(Self);
-		if (MapSelf->MapProp && MapSelf->MapInstance && !MapSelf->OwnerContext.HasOwner())
-		{
-			Collector.AddReferencedObject(MapSelf->MapProp);
-			FPyReferenceCollector::AddReferencedObjectsFromProperty(Collector, MapSelf->MapProp, MapSelf->MapInstance);
-		}
-	};
+		Collector.AddReferencedObject(Self->MapProp);
+		FPyReferenceCollector::AddReferencedObjectsFromProperty(Collector, Self->MapProp, Self->MapInstance);
+	}
 }
 
 #endif	// WITH_PYTHON
diff --git a/Engine/Plugins/Experimental/PythonScriptPlugin/Source/PythonScriptPlugin/Private/PyWrapperMap.h b/Engine/Plugins/Experimental/PythonScriptPlugin/Source/PythonScriptPlugin/Private/PyWrapperMap.h
--- a/Engine/Plugins/Experimental/PythonScriptPlugin/Source/PythonScriptPlugin/Private/PyWrapperMap.h
+++ b/Engine/Plugins/Experimental/PythonScriptPlugin/Source/PythonScriptPlugin/Private/PyWrapperMap.h
@@ -45,10 +45,10 @@ struct FPyWrapperMap : public FPyWrapperBase
 	static bool ValidateInternalState(FPyWrapperMap* InSelf);
 
 	/** Cast the given Python object to this wrapped type (returns a new reference) */
-	static FPyWrapperMap* CastPyObject(PyObject* InPyObject);
+	static FPyWrapperMap* CastPyObject(PyObject* InPyObject, FPyConversionResult* OutCastResult = nullptr);
 
 	/** Cast the given Python object to this wrapped type, or attempt to convert the type into a new wrapped instance (returns a new reference) */
-	static FPyWrapperMap* CastPyObject(PyObject* InPyObject, PyTypeObject* InType, const PyUtil::FPropertyDef& InKeyDef, const PyUtil::FPropertyDef& InValueDef);
+	static FPyWrapperMap* CastPyObject(PyObject* InPyObject, PyTypeObject* InType, const PyUtil::FPropertyDef& InKeyDef, const PyUtil::FPropertyDef& InValueDef, FPyConversionResult* OutCastResult = nullptr);
 
 	/** Get the length of this container (equivalent to 'len(x)' in Python) */
 	static Py_ssize_t Len(FPyWrapperMap* InSelf);
@@ -96,9 +96,14 @@ struct FPyWrapperMap : public FPyWrapperBase
 /** Meta-data for all UE4 exposed map types */
 struct FPyWrapperMapMetaData : public FPyWrapperBaseMetaData
 {
-	PY_OVERRIDE_GETSET_METADATA(FPyWrapperMapMetaData)
+	PY_METADATA_METHODS(FPyWrapperMapMetaData, FGuid(0xA6AAB38C, 0x5C174F81, 0xB80FA903, 0xBF50C8A2))
 
-	FPyWrapperMapMetaData();
+	FPyWrapperMapMetaData()
+	{
+	}
+
+	/** Add object references from the given Python object to the given collector */
+	virtual void AddReferencedObjects(FPyWrapperBase* Instance, FReferenceCollector& Collector) override;
 };
 
 typedef TPyPtr<FPyWrapperMap> FPyWrapperMapPtr;
diff --git a/Engine/Plugins/Experimental/PythonScriptPlugin/Source/PythonScriptPlugin/Private/PyWrapperMath.cpp b/Engine/Plugins/Experimental/PythonScriptPlugin/Source/PythonScriptPlugin/Private/PyWrapperMath.cpp
--- a/Engine/Plugins/Experimental/PythonScriptPlugin/Source/PythonScriptPlugin/Private/PyWrapperMath.cpp
+++ b/Engine/Plugins/Experimental/PythonScriptPlugin/Source/PythonScriptPlugin/Private/PyWrapperMath.cpp
@@ -10,7 +10,7 @@
 #if WITH_PYTHON
 
 template <typename WrapperType>
-void InitializeAndRegisterMathType(PyObject* PyModule, PyTypeObject* InPyType, FPyConstantDef* InPyConstants, const TFunctionRef<void(TArray<PyGenUtil::FGeneratedWrappedMethodParameter>&)>& InBuildInitParamsFunc)
+void InitializeAndRegisterMathType(PyObject* PyModule, PyTypeObject* InPyType, FPyConstantDef* InPyConstants, const TFunctionRef<void(const UStruct*, TArray<PyGenUtil::FGeneratedWrappedMethodParameter>&)>& InBuildInitParamsFunc)
 {
 	typedef typename WrapperType::WrappedType WrappedType;
 
@@ -18,7 +18,7 @@ void InitializeAndRegisterMathType(PyObject* PyModule, PyTypeObject* InPyType, F
 	{
 		static TPyWrapperInlineStructMetaData<WrappedType> MetaData;
 		MetaData.Struct = TBaseStructure<WrappedType>::Get();
-		InBuildInitParamsFunc(MetaData.InitParams);
+		InBuildInitParamsFunc(MetaData.Struct, MetaData.InitParams);
 		TPyWrapperInlineStructMetaData<WrappedType>::SetMetaData(InPyType, &MetaData);
 
 		if (InPyConstants)
@@ -53,11 +53,11 @@ void InitializePyWrapperVector(PyObject* PyModule)
 		{ nullptr, nullptr, nullptr, nullptr }
 	};
 	
-	InitializeAndRegisterMathType<FPyWrapperVector>(PyModule, &PyWrapperVectorType, PyConstants, [](TArray<PyGenUtil::FGeneratedWrappedMethodParameter>& OutInitParams)
+	InitializeAndRegisterMathType<FPyWrapperVector>(PyModule, &PyWrapperVectorType, PyConstants, [](const UStruct* InStruct, TArray<PyGenUtil::FGeneratedWrappedMethodParameter>& OutInitParams)
 	{
-		PyGenUtil::AddStructInitParam(TEXT("X"), TEXT("x"), OutInitParams);
-		PyGenUtil::AddStructInitParam(TEXT("Y"), TEXT("y"), OutInitParams);
-		PyGenUtil::AddStructInitParam(TEXT("Z"), TEXT("z"), OutInitParams);
+		PyGenUtil::AddStructInitParam(InStruct->FindPropertyByName(TEXT("X")), TEXT("x"), OutInitParams);
+		PyGenUtil::AddStructInitParam(InStruct->FindPropertyByName(TEXT("Y")), TEXT("y"), OutInitParams);
+		PyGenUtil::AddStructInitParam(InStruct->FindPropertyByName(TEXT("Z")), TEXT("z"), OutInitParams);
 	});
 }
 
@@ -78,21 +78,21 @@ void InitializePyWrapperVector2D(PyObject* PyModule)
 		{ nullptr, nullptr, nullptr, nullptr }
 	};
 
-	InitializeAndRegisterMathType<FPyWrapperVector2D>(PyModule, &PyWrapperVector2DType, PyConstants, [](TArray<PyGenUtil::FGeneratedWrappedMethodParameter>& OutInitParams)
+	InitializeAndRegisterMathType<FPyWrapperVector2D>(PyModule, &PyWrapperVector2DType, PyConstants, [](const UStruct* InStruct, TArray<PyGenUtil::FGeneratedWrappedMethodParameter>& OutInitParams)
 	{
-		PyGenUtil::AddStructInitParam(TEXT("X"), TEXT("x"), OutInitParams);
-		PyGenUtil::AddStructInitParam(TEXT("Y"), TEXT("y"), OutInitParams);
+		PyGenUtil::AddStructInitParam(InStruct->FindPropertyByName(TEXT("X")), TEXT("x"), OutInitParams);
+		PyGenUtil::AddStructInitParam(InStruct->FindPropertyByName(TEXT("Y")), TEXT("y"), OutInitParams);
 	});
 }
 
 void InitializePyWrapperVector4(PyObject* PyModule)
 {
-	InitializeAndRegisterMathType<FPyWrapperVector4>(PyModule, &PyWrapperVector4Type, nullptr, [](TArray<PyGenUtil::FGeneratedWrappedMethodParameter>& OutInitParams)
+	InitializeAndRegisterMathType<FPyWrapperVector4>(PyModule, &PyWrapperVector4Type, nullptr, [](const UStruct* InStruct, TArray<PyGenUtil::FGeneratedWrappedMethodParameter>& OutInitParams)
 	{
-		PyGenUtil::AddStructInitParam(TEXT("X"), TEXT("x"), OutInitParams);
-		PyGenUtil::AddStructInitParam(TEXT("Y"), TEXT("y"), OutInitParams);
-		PyGenUtil::AddStructInitParam(TEXT("Z"), TEXT("z"), OutInitParams);
-		PyGenUtil::AddStructInitParam(TEXT("W"), TEXT("w"), OutInitParams);
+		PyGenUtil::AddStructInitParam(InStruct->FindPropertyByName(TEXT("X")), TEXT("x"), OutInitParams);
+		PyGenUtil::AddStructInitParam(InStruct->FindPropertyByName(TEXT("Y")), TEXT("y"), OutInitParams);
+		PyGenUtil::AddStructInitParam(InStruct->FindPropertyByName(TEXT("Z")), TEXT("z"), OutInitParams);
+		PyGenUtil::AddStructInitParam(InStruct->FindPropertyByName(TEXT("W")), TEXT("w"), OutInitParams);
 	});
 }
 
@@ -112,12 +112,12 @@ void InitializePyWrapperQuat(PyObject* PyModule)
 		{ nullptr, nullptr, nullptr, nullptr }
 	};
 
-	InitializeAndRegisterMathType<FPyWrapperQuat>(PyModule, &PyWrapperQuatType, PyConstants, [](TArray<PyGenUtil::FGeneratedWrappedMethodParameter>& OutInitParams)
+	InitializeAndRegisterMathType<FPyWrapperQuat>(PyModule, &PyWrapperQuatType, PyConstants, [](const UStruct* InStruct, TArray<PyGenUtil::FGeneratedWrappedMethodParameter>& OutInitParams)
 	{
-		PyGenUtil::AddStructInitParam(TEXT("X"), TEXT("x"), OutInitParams);
-		PyGenUtil::AddStructInitParam(TEXT("Y"), TEXT("y"), OutInitParams);
-		PyGenUtil::AddStructInitParam(TEXT("Z"), TEXT("z"), OutInitParams);
-		PyGenUtil::AddStructInitParam(TEXT("W"), TEXT("w"), OutInitParams);
+		PyGenUtil::AddStructInitParam(InStruct->FindPropertyByName(TEXT("X")), TEXT("x"), OutInitParams);
+		PyGenUtil::AddStructInitParam(InStruct->FindPropertyByName(TEXT("Y")), TEXT("y"), OutInitParams);
+		PyGenUtil::AddStructInitParam(InStruct->FindPropertyByName(TEXT("Z")), TEXT("z"), OutInitParams);
+		PyGenUtil::AddStructInitParam(InStruct->FindPropertyByName(TEXT("W")), TEXT("w"), OutInitParams);
 	});
 }
 
@@ -144,12 +144,12 @@ void InitializePyWrapperLinearColor(PyObject* PyModule)
 		{ nullptr, nullptr, nullptr, nullptr }
 	};
 
-	InitializeAndRegisterMathType<FPyWrapperLinearColor>(PyModule, &PyWrapperLinearColorType, PyConstants, [](TArray<PyGenUtil::FGeneratedWrappedMethodParameter>& OutInitParams)
+	InitializeAndRegisterMathType<FPyWrapperLinearColor>(PyModule, &PyWrapperLinearColorType, PyConstants, [](const UStruct* InStruct, TArray<PyGenUtil::FGeneratedWrappedMethodParameter>& OutInitParams)
 	{
-		PyGenUtil::AddStructInitParam(TEXT("R"), TEXT("r"), OutInitParams);
-		PyGenUtil::AddStructInitParam(TEXT("G"), TEXT("g"), OutInitParams);
-		PyGenUtil::AddStructInitParam(TEXT("B"), TEXT("b"), OutInitParams);
-		PyGenUtil::AddStructInitParam(TEXT("A"), TEXT("a"), OutInitParams);
+		PyGenUtil::AddStructInitParam(InStruct->FindPropertyByName(TEXT("R")), TEXT("r"), OutInitParams);
+		PyGenUtil::AddStructInitParam(InStruct->FindPropertyByName(TEXT("G")), TEXT("g"), OutInitParams);
+		PyGenUtil::AddStructInitParam(InStruct->FindPropertyByName(TEXT("B")), TEXT("b"), OutInitParams);
+		PyGenUtil::AddStructInitParam(InStruct->FindPropertyByName(TEXT("A")), TEXT("a"), OutInitParams);
 	});
 }
 
@@ -248,7 +248,7 @@ bool StructStructOp_ReturnStruct(PyTypeObject* InStructType, WrapperType* InLHS,
 	}
 
 	WrappedType Result = OpType::template Apply<WrappedType, WrappedType, WrappedType>(WrapperType::GetTypedStruct(InLHS), WrapperType::GetTypedStruct(RHS));
-	return PyConversion::PythonizeStruct(Result, OutResult);
+	return PyConversion::PythonizeStruct(Result, OutResult).Succeeded();
 }
 
 /** Struct(op=)Struct -> Struct */
@@ -285,7 +285,7 @@ bool StructStructOp_ReturnIntrinsic(PyTypeObject* InStructType, WrapperType* InL
 	}
 
 	IntrinsicType Result = OpType::template Apply<IntrinsicType, WrappedType, WrappedType>(WrapperType::GetTypedStruct(InLHS), WrapperType::GetTypedStruct(RHS));
-	return PyConversion::Pythonize(Result, OutResult);
+	return PyConversion::Pythonize(Result, OutResult).Succeeded();
 }
 
 /** Struct(op)Intrinsic -> Struct */
@@ -301,7 +301,7 @@ bool StructIntrinsicOp_ReturnStruct(PyTypeObject* InStructType, WrapperType* InL
 	}
 
 	WrappedType Result = OpType::template Apply<WrappedType, WrappedType, IntrinsicType>(WrapperType::GetTypedStruct(InLHS), RHS);
-	return PyConversion::PythonizeStruct(Result, OutResult);
+	return PyConversion::PythonizeStruct(Result, OutResult).Succeeded();
 }
 
 /** Struct(op=)Intrinsic -> Struct */
diff --git a/Engine/Plugins/Experimental/PythonScriptPlugin/Source/PythonScriptPlugin/Private/PyWrapperName.cpp b/Engine/Plugins/Experimental/PythonScriptPlugin/Source/PythonScriptPlugin/Private/PyWrapperName.cpp
--- a/Engine/Plugins/Experimental/PythonScriptPlugin/Source/PythonScriptPlugin/Private/PyWrapperName.cpp
+++ b/Engine/Plugins/Experimental/PythonScriptPlugin/Source/PythonScriptPlugin/Private/PyWrapperName.cpp
@@ -16,63 +16,26 @@ void InitializePyWrapperName(PyObject* PyModule)
 	}
 }
 
-FPyWrapperName* FPyWrapperName::New(PyTypeObject* InType)
+void FPyWrapperName::InitValue(FPyWrapperName* InSelf, const FName InValue)
 {
-	FPyWrapperName* Self = (FPyWrapperName*)FPyWrapperBase::New(InType);
-	if (Self)
-	{
-		new(&Self->Name) FName();
-	}
-	return Self;
-}
-
-void FPyWrapperName::Free(FPyWrapperName* InSelf)
-{
-	Deinit(InSelf);
-
-	InSelf->Name.~FName();
-	FPyWrapperBase::Free(InSelf);
+	Super::InitValue(InSelf, InValue);
+	FPyWrapperNameFactory::Get().MapInstance(InSelf->Value, InSelf);
 }
 
-int FPyWrapperName::Init(FPyWrapperName* InSelf)
+void FPyWrapperName::DeinitValue(FPyWrapperName* InSelf)
 {
-	Deinit(InSelf);
-
-	const int BaseInit = FPyWrapperBase::Init(InSelf);
-	if (BaseInit != 0)
-	{
-		return BaseInit;
-	}
-
-	FPyWrapperNameFactory::Get().MapInstance(InSelf->Name, InSelf);
-	return 0;
-}
-
-int FPyWrapperName::Init(FPyWrapperName* InSelf, const FName InValue)
-{
-	Deinit(InSelf);
-
-	const int BaseInit = FPyWrapperBase::Init(InSelf);
-	if (BaseInit != 0)
-	{
-		return BaseInit;
-	}
-
-	InSelf->Name = InValue;
-	FPyWrapperNameFactory::Get().MapInstance(InSelf->Name, InSelf);
-	return 0;
+	FPyWrapperNameFactory::Get().UnmapInstance(InSelf->Value, Py_TYPE(InSelf));
+	Super::DeinitValue(InSelf);
 }
 
-void FPyWrapperName::Deinit(FPyWrapperName* InSelf)
+FPyWrapperName* FPyWrapperName::CastPyObject(PyObject* InPyObject, FPyConversionResult* OutCastResult)
 {
-	FPyWrapperNameFactory::Get().UnmapInstance(InSelf->Name, Py_TYPE(InSelf));
-	InSelf->Name = FName();
-}
+	SetOptionalPyConversionResult(FPyConversionResult::Failure(), OutCastResult);
 
-FPyWrapperName* FPyWrapperName::CastPyObject(PyObject* InPyObject)
-{
 	if (PyObject_IsInstance(InPyObject, (PyObject*)&PyWrapperNameType) == 1)
 	{
+		SetOptionalPyConversionResult(FPyConversionResult::Success(), OutCastResult);
+
 		Py_INCREF(InPyObject);
 		return (FPyWrapperName*)InPyObject;
 	}
@@ -80,10 +43,14 @@ FPyWrapperName* FPyWrapperName::CastPyObject(PyObject* InPyObject)
 	return nullptr;
 }
 
-FPyWrapperName* FPyWrapperName::CastPyObject(PyObject* InPyObject, PyTypeObject* InType)
+FPyWrapperName* FPyWrapperName::CastPyObject(PyObject* InPyObject, PyTypeObject* InType, FPyConversionResult* OutCastResult)
 {
+	SetOptionalPyConversionResult(FPyConversionResult::Failure(), OutCastResult);
+
 	if (PyObject_IsInstance(InPyObject, (PyObject*)InType) == 1 && (InType == &PyWrapperNameType || PyObject_IsInstance(InPyObject, (PyObject*)&PyWrapperNameType) == 1))
 	{
+		SetOptionalPyConversionResult(Py_TYPE(InPyObject) == InType ? FPyConversionResult::Success() : FPyConversionResult::SuccessWithCoercion(), OutCastResult);
+
 		Py_INCREF(InPyObject);
 		return (FPyWrapperName*)InPyObject;
 	}
@@ -100,6 +67,7 @@ FPyWrapperName* FPyWrapperName::CastPyObject(PyObject* InPyObject, PyTypeObject*
 					return nullptr;
 				}
 			}
+			SetOptionalPyConversionResult(FPyConversionResult::SuccessWithCoercion(), OutCastResult);
 			return NewName.Release();
 		}
 	}
@@ -111,16 +79,6 @@ PyTypeObject InitializePyWrapperNameType()
 {
 	struct FFuncs
 	{
-		static PyObject* New(PyTypeObject* InType, PyObject* InArgs, PyObject* InKwds)
-		{
-			return (PyObject*)FPyWrapperName::New(InType);
-		}
-
-		static void Dealloc(FPyWrapperName* InSelf)
-		{
-			FPyWrapperName::Free(InSelf);
-		}
-
 		static int Init(FPyWrapperName* InSelf, PyObject* InArgs, PyObject* InKwds)
 		{
 			FName InitValue;
@@ -140,7 +98,7 @@ PyTypeObject InitializePyWrapperNameType()
 
 		static PyObject* Str(FPyWrapperName* InSelf)
 		{
-			return PyUnicode_FromString(TCHAR_TO_UTF8(*InSelf->Name.ToString()));
+			return PyUnicode_FromString(TCHAR_TO_UTF8(*InSelf->Value.ToString()));
 		}
 
 		static PyObject* RichCmp(FPyWrapperName* InSelf, PyObject* InOther, int InOp)
@@ -152,12 +110,12 @@ PyTypeObject InitializePyWrapperNameType()
 				return Py_NotImplemented;
 			}
 
-			return PyUtil::PyRichCmp(InSelf->Name.Compare(Other), 0, InOp);
+			return PyUtil::PyRichCmp(InSelf->Value.Compare(Other), 0, InOp);
 		}
 
 		static PyUtil::FPyHashType Hash(FPyWrapperName* InSelf)
 		{
-			const PyUtil::FPyHashType PyHash = (PyUtil::FPyHashType)GetTypeHash(InSelf->Name);
+			const PyUtil::FPyHashType PyHash = (PyUtil::FPyHashType)GetTypeHash(InSelf->Value);
 			return PyHash != -1 ? PyHash : 0;
 		}
 	};
@@ -182,7 +140,7 @@ PyTypeObject InitializePyWrapperNameType()
 
 		static PyObject* IsValid(FPyWrapperName* InSelf)
 		{
-			if (InSelf->Name.IsValid())
+			if (InSelf->Value.IsValid())
 			{
 				Py_RETURN_TRUE;
 			}
@@ -192,7 +150,7 @@ PyTypeObject InitializePyWrapperNameType()
 
 		static PyObject* IsNone(FPyWrapperName* InSelf)
 		{
-			if (InSelf->Name.IsNone())
+			if (InSelf->Value.IsNone())
 			{
 				Py_RETURN_TRUE;
 			}
@@ -208,15 +166,8 @@ PyTypeObject InitializePyWrapperNameType()
 		{ nullptr, nullptr, 0, nullptr }
 	};
 
-	PyTypeObject PyType = {
-		PyVarObject_HEAD_INIT(nullptr, 0)
-		"Name", /* tp_name */
-		sizeof(FPyWrapperName), /* tp_basicsize */
-	};
+	PyTypeObject PyType = InitializePyWrapperBasicType<FPyWrapperName>("Name", "Type for all UE4 exposed FName instances");
 
-	PyType.tp_base = &PyWrapperBaseType;
-	PyType.tp_new = (newfunc)&FFuncs::New;
-	PyType.tp_dealloc = (destructor)&FFuncs::Dealloc;
 	PyType.tp_init = (initproc)&FFuncs::Init;
 	PyType.tp_str = (reprfunc)&FFuncs::Str;
 	PyType.tp_richcompare = (richcmpfunc)&FFuncs::RichCmp;
@@ -224,9 +175,6 @@ PyTypeObject InitializePyWrapperNameType()
 
 	PyType.tp_methods = PyMethods;
 
-	PyType.tp_flags = Py_TPFLAGS_DEFAULT;
-	PyType.tp_doc = "Type for all UE4 exposed FName instances";
-
 	return PyType;
 }
 
diff --git a/Engine/Plugins/Experimental/PythonScriptPlugin/Source/PythonScriptPlugin/Private/PyWrapperName.h b/Engine/Plugins/Experimental/PythonScriptPlugin/Source/PythonScriptPlugin/Private/PyWrapperName.h
--- a/Engine/Plugins/Experimental/PythonScriptPlugin/Source/PythonScriptPlugin/Private/PyWrapperName.h
+++ b/Engine/Plugins/Experimental/PythonScriptPlugin/Source/PythonScriptPlugin/Private/PyWrapperName.h
@@ -2,7 +2,7 @@
 
 #pragma once
 
-#include "PyWrapperBase.h"
+#include "PyWrapperBasic.h"
 
 #if WITH_PYTHON
 
@@ -13,31 +13,21 @@ extern PyTypeObject PyWrapperNameType;
 void InitializePyWrapperName(PyObject* PyModule);
 
 /** Type for all UE4 exposed FName instances */
-struct FPyWrapperName : public FPyWrapperBase
+struct FPyWrapperName : public TPyWrapperBasic<FName, FPyWrapperName>
 {
-	/** The wrapped name */
-	FName Name;
-
-	/** New this wrapper instance (called via tp_new for Python, or directly in C++) */
-	static FPyWrapperName* New(PyTypeObject* InType);
-
-	/** Free this wrapper instance (called via tp_dealloc for Python) */
-	static void Free(FPyWrapperName* InSelf);
-
-	/** Initialize this wrapper instance (called via tp_init for Python, or directly in C++) */
-	static int Init(FPyWrapperName* InSelf);
-
-	/** Initialize this wrapper instance to the given value (called via tp_init for Python, or directly in C++) */
-	static int Init(FPyWrapperName* InSelf, const FName InValue);
-
-	/** Deinitialize this wrapper instance (called via Init and Free to restore the instance to its New state) */
-	static void Deinit(FPyWrapperName* InSelf);
+	typedef TPyWrapperBasic<FName, FPyWrapperName> Super;
 
 	/** Cast the given Python object to this wrapped type (returns a new reference) */
-	static FPyWrapperName* CastPyObject(PyObject* InPyObject);
+	static FPyWrapperName* CastPyObject(PyObject* InPyObject, FPyConversionResult* OutCastResult = nullptr);
 
 	/** Cast the given Python object to this wrapped type, or attempt to convert the type into a new wrapped instance (returns a new reference) */
-	static FPyWrapperName* CastPyObject(PyObject* InPyObject, PyTypeObject* InType);
+	static FPyWrapperName* CastPyObject(PyObject* InPyObject, PyTypeObject* InType, FPyConversionResult* OutCastResult = nullptr);
+
+	/** Initialize the value of this wrapper instance (internal) */
+	static void InitValue(FPyWrapperName* InSelf, const FName InValue);
+
+	/** Deinitialize the value of this wrapper instance (internal) */
+	static void DeinitValue(FPyWrapperName* InSelf);
 };
 
 typedef TPyPtr<FPyWrapperName> FPyWrapperNamePtr;
diff --git a/Engine/Plugins/Experimental/PythonScriptPlugin/Source/PythonScriptPlugin/Private/PyWrapperObject.cpp b/Engine/Plugins/Experimental/PythonScriptPlugin/Source/PythonScriptPlugin/Private/PyWrapperObject.cpp
--- a/Engine/Plugins/Experimental/PythonScriptPlugin/Source/PythonScriptPlugin/Private/PyWrapperObject.cpp
+++ b/Engine/Plugins/Experimental/PythonScriptPlugin/Source/PythonScriptPlugin/Private/PyWrapperObject.cpp
@@ -84,10 +84,14 @@ bool FPyWrapperObject::ValidateInternalState(FPyWrapperObject* InSelf)
 	return true;
 }
 
-FPyWrapperObject* FPyWrapperObject::CastPyObject(PyObject* InPyObject)
+FPyWrapperObject* FPyWrapperObject::CastPyObject(PyObject* InPyObject, FPyConversionResult* OutCastResult)
 {
+	SetOptionalPyConversionResult(FPyConversionResult::Failure(), OutCastResult);
+
 	if (PyObject_IsInstance(InPyObject, (PyObject*)&PyWrapperObjectType) == 1)
 	{
+		SetOptionalPyConversionResult(FPyConversionResult::Success(), OutCastResult);
+
 		Py_INCREF(InPyObject);
 		return (FPyWrapperObject*)InPyObject;
 	}
@@ -95,10 +99,14 @@ FPyWrapperObject* FPyWrapperObject::CastPyObject(PyObject* InPyObject)
 	return nullptr;
 }
 
-FPyWrapperObject* FPyWrapperObject::CastPyObject(PyObject* InPyObject, PyTypeObject* InType)
+FPyWrapperObject* FPyWrapperObject::CastPyObject(PyObject* InPyObject, PyTypeObject* InType, FPyConversionResult* OutCastResult)
 {
+	SetOptionalPyConversionResult(FPyConversionResult::Failure(), OutCastResult);
+
 	if (PyObject_IsInstance(InPyObject, (PyObject*)InType) == 1 && (InType == &PyWrapperObjectType || PyObject_IsInstance(InPyObject, (PyObject*)&PyWrapperObjectType) == 1))
 	{
+		SetOptionalPyConversionResult(Py_TYPE(InPyObject) == InType ? FPyConversionResult::Success() : FPyConversionResult::SuccessWithCoercion(), OutCastResult);
+
 		Py_INCREF(InPyObject);
 		return (FPyWrapperObject*)InPyObject;
 	}
@@ -106,269 +114,269 @@ FPyWrapperObject* FPyWrapperObject::CastPyObject(PyObject* InPyObject, PyTypeObj
 	return nullptr;
 }
 
-PyObject* FPyWrapperObject::GetPropertyValue(FPyWrapperObject* InSelf, const FName InPropName, const char* InPythonAttrName)
+PyObject* FPyWrapperObject::GetPropertyValueByName(FPyWrapperObject* InSelf, const FName InPropName, const char* InPythonAttrName)
 {
 	if (!ValidateInternalState(InSelf))
 	{
 		return nullptr;
 	}
 
-	return PyUtil::GetUEPropValue(InSelf->ObjectInstance->GetClass(), InSelf->ObjectInstance, InPropName, InPythonAttrName, (PyObject*)InSelf, *PyUtil::GetErrorContext(InSelf));
+	UClass* Class = InSelf->ObjectInstance->GetClass();
+
+	UProperty* Prop = Class->FindPropertyByName(InPropName);
+	if (!Prop)
+	{
+		PyUtil::SetPythonError(PyExc_Exception, InSelf, *FString::Printf(TEXT("Failed to find property '%s' for attribute '%s' on '%s'"), *InPropName.ToString(), UTF8_TO_TCHAR(InPythonAttrName), *Class->GetName()));
+		return nullptr;
+	}
+
+	return PyUtil::GetUEPropValue(Class, InSelf->ObjectInstance, Prop, InPythonAttrName, (PyObject*)InSelf, *PyUtil::GetErrorContext(InSelf));
 }
 
-int FPyWrapperObject::SetPropertyValue(FPyWrapperObject* InSelf, PyObject* InValue, const FName InPropName, const char* InPythonAttrName, const bool InNotifyChange, const uint64 InReadOnlyFlags)
+PyObject* FPyWrapperObject::GetPropertyValue(FPyWrapperObject* InSelf, const UProperty* InProp, const char* InPythonAttrName)
 {
 	if (!ValidateInternalState(InSelf))
 	{
-		return -1;
+		return nullptr;
 	}
 
-	const FPyWrapperOwnerContext ChangeOwner = InNotifyChange ? FPyWrapperOwnerContext((PyObject*)InSelf, InPropName) : FPyWrapperOwnerContext();
-	return PyUtil::SetUEPropValue(InSelf->ObjectInstance->GetClass(), InSelf->ObjectInstance, InValue, InPropName, InPythonAttrName, ChangeOwner, InReadOnlyFlags, *PyUtil::GetErrorContext(InSelf));
+	return PyUtil::GetUEPropValue(InSelf->ObjectInstance->GetClass(), InSelf->ObjectInstance, InProp, InPythonAttrName, (PyObject*)InSelf, *PyUtil::GetErrorContext(InSelf));
 }
 
-PyObject* FPyWrapperObject::CallGetterFunction(UClass* InClass, FPyWrapperObject* InSelf, const FName InFuncName)
+int FPyWrapperObject::SetPropertyValueByName(FPyWrapperObject* InSelf, PyObject* InValue, const FName InPropName, const char* InPythonAttrName, const bool InNotifyChange, const uint64 InReadOnlyFlags)
 {
 	if (!ValidateInternalState(InSelf))
 	{
-		return nullptr;
+		return -1;
+	}
+
+	UClass* Class = InSelf->ObjectInstance->GetClass();
+
+	UProperty* Prop = Class->FindPropertyByName(InPropName);
+	if (!Prop)
+	{
+		PyUtil::SetPythonError(PyExc_Exception, InSelf, *FString::Printf(TEXT("Failed to find property '%s' for attribute '%s' on '%s'"), *InPropName.ToString(), UTF8_TO_TCHAR(InPythonAttrName), *Class->GetName()));
+		return -1;
 	}
 
-	return CallFunction_Impl(InClass, InSelf->ObjectInstance, InFuncName, *PyUtil::GetErrorContext(InSelf));
+	const FPyWrapperOwnerContext ChangeOwner = InNotifyChange ? FPyWrapperOwnerContext((PyObject*)InSelf, Prop) : FPyWrapperOwnerContext();
+	return PyUtil::SetUEPropValue(Class, InSelf->ObjectInstance, InValue, Prop, InPythonAttrName, ChangeOwner, InReadOnlyFlags, *PyUtil::GetErrorContext(InSelf));
 }
 
-int FPyWrapperObject::CallSetterFunction(UClass* InClass, FPyWrapperObject* InSelf, PyObject* InValue, const FName InFuncName)
+int FPyWrapperObject::SetPropertyValue(FPyWrapperObject* InSelf, PyObject* InValue, const UProperty* InProp, const char* InPythonAttrName, const bool InNotifyChange, const uint64 InReadOnlyFlags)
 {
 	if (!ValidateInternalState(InSelf))
 	{
 		return -1;
 	}
 
-	UFunction* Func = InClass->FindFunctionByName(InFuncName);
-	if (ensureAlways(Func))
+	const FPyWrapperOwnerContext ChangeOwner = InNotifyChange ? FPyWrapperOwnerContext((PyObject*)InSelf, InProp) : FPyWrapperOwnerContext();
+	return PyUtil::SetUEPropValue(InSelf->ObjectInstance->GetClass(), InSelf->ObjectInstance, InValue, InProp, InPythonAttrName, ChangeOwner, InReadOnlyFlags, *PyUtil::GetErrorContext(InSelf));
+}
+
+PyObject* FPyWrapperObject::CallGetterFunction(FPyWrapperObject* InSelf, const PyGenUtil::FGeneratedWrappedFunction& InFuncDef)
+{
+	if (!ValidateInternalState(InSelf))
 	{
-		check(InSelf->ObjectInstance->IsA(InClass) || InSelf->ObjectInstance->GetClass()->ImplementsInterface(InClass));
+		return nullptr;
+	}
+
+	return CallFunction_Impl(InSelf->ObjectInstance, InFuncDef, *PyUtil::GetErrorContext(InSelf));
+}
 
-		FStructOnScope FuncParams(Func);
+int FPyWrapperObject::CallSetterFunction(FPyWrapperObject* InSelf, PyObject* InValue, const PyGenUtil::FGeneratedWrappedFunction& InFuncDef)
+{
+	if (!ValidateInternalState(InSelf))
+	{
+		return -1;
+	}
 
-		UProperty* ArgProp = nullptr;
-		for (TFieldIterator<UProperty> ParamIt(FuncParams.GetStruct()); ParamIt; ++ParamIt)
+	if (ensureAlways(InFuncDef.Func))
+	{
+		// Deprecated functions emit a warning
 		{
-			UProperty* Param = *ParamIt;
-			if (!Param->HasAnyPropertyFlags(CPF_ReturnParm))
+			FString DeprecationMessage;
+			if (PyGenUtil::IsDeprecatedFunction(InFuncDef.Func, &DeprecationMessage) &&
+				PyUtil::SetPythonWarning(PyExc_DeprecationWarning, InSelf, *FString::Printf(TEXT("Function '%s.%s' is deprecated: %s"), *InFuncDef.Func->GetOwnerClass()->GetName(), *InFuncDef.Func->GetName(), *DeprecationMessage)) == -1
+				)
 			{
-				ArgProp = Param;
-				break;
+				// -1 from SetPythonWarning means the warning should be an exception
+				return -1;
 			}
 		}
 
-		if (!ArgProp)
+		// Setter functions should have a single input parameter and no output parameters
+		if (InFuncDef.InputParams.Num() != 1 || InFuncDef.OutputParams.Num() != 0)
 		{
-			PyUtil::SetPythonError(PyExc_Exception, InSelf, *FString::Printf(TEXT("Failed to find input parameter when calling function '%s' on '%s'"), *Func->GetName(), *InSelf->ObjectInstance->GetName()));
+			PyUtil::SetPythonError(PyExc_Exception, InSelf, *FString::Printf(TEXT("Setter function '%s.%s' on '%s' has the incorrect number of parameters (expected 1 input and 0 output, got %d input and %d output)"), *InFuncDef.Func->GetOwnerClass()->GetName(), *InFuncDef.Func->GetName(), *InSelf->ObjectInstance->GetName(), InFuncDef.InputParams.Num(), InFuncDef.OutputParams.Num()));
 			return -1;
 		}
 
+		FStructOnScope FuncParams(InFuncDef.Func);
 		if (InValue)
 		{
-			if (!PyConversion::NativizeProperty_InContainer(InValue, ArgProp, FuncParams.GetStructMemory(), 0))
+			if (!PyConversion::NativizeProperty_InContainer(InValue, InFuncDef.InputParams[0].ParamProp, FuncParams.GetStructMemory(), 0))
 			{
-				PyUtil::SetPythonError(PyExc_TypeError, InSelf, *FString::Printf(TEXT("Failed to convert input parameter when calling function '%s' on '%s'"), *Func->GetName(), *InSelf->ObjectInstance->GetName()));
+				PyUtil::SetPythonError(PyExc_TypeError, InSelf, *FString::Printf(TEXT("Failed to convert input parameter when calling function '%s.%s' on '%s'"), *InFuncDef.Func->GetOwnerClass()->GetName(), *InFuncDef.Func->GetName(), *InSelf->ObjectInstance->GetName()));
 				return -1;
 			}
 		}
-
-		CallFunction_InvokeImpl(InSelf->ObjectInstance, Func, FuncParams.GetStructMemory(), *PyUtil::GetErrorContext(InSelf));
+		PyUtil::InvokeFunctionCall(InSelf->ObjectInstance, InFuncDef.Func, FuncParams.GetStructMemory(), *PyUtil::GetErrorContext(InSelf));
 	}
 
 	return 0;
 }
 
-PyObject* FPyWrapperObject::CallFunction(UClass* InClass, PyTypeObject* InType, const FName InFuncName)
+PyObject* FPyWrapperObject::CallFunction(PyTypeObject* InType, const PyGenUtil::FGeneratedWrappedFunction& InFuncDef)
 {
-	UObject* Obj = InClass ? InClass->GetDefaultObject() : nullptr;
-	return CallFunction_Impl(InClass, Obj, InFuncName, *PyUtil::GetErrorContext(InType));
+	UClass* Class = FPyWrapperObjectMetaData::GetClass(InType);
+	UObject* Obj = Class ? Class->GetDefaultObject() : nullptr;
+	return CallFunction_Impl(Obj, InFuncDef, *PyUtil::GetErrorContext(InType));
 }
 
-PyObject* FPyWrapperObject::CallFunction(UClass* InClass, PyTypeObject* InType, PyObject* InArgs, PyObject* InKwds, const FName InFuncName, const char* InPythonFuncName, const TArray<PyGenUtil::FGeneratedWrappedMethodParameter>& InParamDef)
+PyObject* FPyWrapperObject::CallFunction(PyTypeObject* InType, PyObject* InArgs, PyObject* InKwds, const PyGenUtil::FGeneratedWrappedFunction& InFuncDef, const char* InPythonFuncName)
 {
-	UObject* Obj = InClass ? InClass->GetDefaultObject() : nullptr;
-	return CallFunction_Impl(InClass, Obj, InArgs, InKwds, InFuncName, InPythonFuncName, InParamDef, *PyUtil::GetErrorContext(InType));
+	UClass* Class = FPyWrapperObjectMetaData::GetClass(InType);
+	UObject* Obj = Class ? Class->GetDefaultObject() : nullptr;
+	return CallFunction_Impl(Obj, InArgs, InKwds, InFuncDef, InPythonFuncName, *PyUtil::GetErrorContext(InType));
 }
 
-PyObject* FPyWrapperObject::CallFunction(UClass* InClass, FPyWrapperObject* InSelf, const FName InFuncName)
+PyObject* FPyWrapperObject::CallFunction(FPyWrapperObject* InSelf, const PyGenUtil::FGeneratedWrappedFunction& InFuncDef)
 {
 	if (!ValidateInternalState(InSelf))
 	{
 		return nullptr;
 	}
 
-	return CallFunction_Impl(InClass, InSelf->ObjectInstance, InFuncName, *PyUtil::GetErrorContext(InSelf));
+	return CallFunction_Impl(InSelf->ObjectInstance, InFuncDef, *PyUtil::GetErrorContext(InSelf));
 }
 
-PyObject* FPyWrapperObject::CallFunction(UClass* InClass, FPyWrapperObject* InSelf, PyObject* InArgs, PyObject* InKwds, const FName InFuncName, const char* InPythonFuncName, const TArray<PyGenUtil::FGeneratedWrappedMethodParameter>& InParamDef)
+PyObject* FPyWrapperObject::CallFunction(FPyWrapperObject* InSelf, PyObject* InArgs, PyObject* InKwds, const PyGenUtil::FGeneratedWrappedFunction& InFuncDef, const char* InPythonFuncName)
 {
 	if (!ValidateInternalState(InSelf))
 	{
 		return nullptr;
 	}
 
-	return CallFunction_Impl(InClass, InSelf->ObjectInstance, InArgs, InKwds, InFuncName, InPythonFuncName, InParamDef, *PyUtil::GetErrorContext(InSelf));
+	return CallFunction_Impl(InSelf->ObjectInstance, InArgs, InKwds, InFuncDef, InPythonFuncName, *PyUtil::GetErrorContext(InSelf));
 }
 
-PyObject* FPyWrapperObject::CallFunction_Impl(UClass* InClass, UObject* InObj, const FName InFuncName, const TCHAR* InErrorCtxt)
+PyObject* FPyWrapperObject::CallFunction_Impl(UObject* InObj, const PyGenUtil::FGeneratedWrappedFunction& InFuncDef, const TCHAR* InErrorCtxt)
 {
-	if (InClass && InObj)
+	if (InObj && ensureAlways(InFuncDef.Func))
 	{
-		check(InObj->IsA(InClass) || InObj->GetClass()->ImplementsInterface(InClass));
-
-		// Find the function from the class rather than the object, as we may be specifically calling a super class function
-		UFunction* Func = InClass->FindFunctionByName(InFuncName);
-		if (ensureAlways(Func))
+		// Deprecated functions emit a warning
 		{
-			if (Func->Children == nullptr)
+			FString DeprecationMessage;
+			if (PyGenUtil::IsDeprecatedFunction(InFuncDef.Func, &DeprecationMessage) &&
+				PyUtil::SetPythonWarning(PyExc_DeprecationWarning, InErrorCtxt, *FString::Printf(TEXT("Function '%s.%s' is deprecated: %s"), *InFuncDef.Func->GetOwnerClass()->GetName(), *InFuncDef.Func->GetName(), *DeprecationMessage)) == -1
+				)
 			{
-				// No return value
-				CallFunction_InvokeImpl(InObj, Func, nullptr, InErrorCtxt);
-			}
-			else
-			{
-				// Return value requires that we create a params struct to hold the result
-				FStructOnScope FuncParams(Func);
-				CallFunction_InvokeImpl(InObj, Func, FuncParams.GetStructMemory(), InErrorCtxt);
-				return CallFunction_ReturnImpl(InObj, Func, FuncParams.GetStructMemory(), InErrorCtxt);
+				// -1 from SetPythonWarning means the warning should be an exception
+				return nullptr;
 			}
 		}
+
+		if (InFuncDef.Func->Children == nullptr)
+		{
+			// No return value
+			PyUtil::InvokeFunctionCall(InObj, InFuncDef.Func, nullptr, InErrorCtxt);
+		}
+		else
+		{
+			// Return value requires that we create a params struct to hold the result
+			FStructOnScope FuncParams(InFuncDef.Func);
+			PyUtil::InvokeFunctionCall(InObj, InFuncDef.Func, FuncParams.GetStructMemory(), InErrorCtxt);
+			return PyGenUtil::PackReturnValues(FuncParams.GetStructMemory(), InFuncDef.OutputParams, InErrorCtxt, *FString::Printf(TEXT("function '%s.%s' on '%s'"), *InFuncDef.Func->GetOwnerClass()->GetName(), *InFuncDef.Func->GetName(), *InObj->GetName()));
+		}
 	}
 
 	Py_RETURN_NONE;
 }
 
-PyObject* FPyWrapperObject::CallFunction_Impl(UClass* InClass, UObject* InObj, PyObject* InArgs, PyObject* InKwds, const FName InFuncName, const char* InPythonFuncName, const TArray<PyGenUtil::FGeneratedWrappedMethodParameter>& InParamDef, const TCHAR* InErrorCtxt)
+PyObject* FPyWrapperObject::CallFunction_Impl(UObject* InObj, PyObject* InArgs, PyObject* InKwds, const PyGenUtil::FGeneratedWrappedFunction& InFuncDef, const char* InPythonFuncName, const TCHAR* InErrorCtxt)
 {
 	TArray<PyObject*> Params;
-	if (!PyGenUtil::ParseMethodParameters(InArgs, InKwds, InParamDef, InPythonFuncName, Params))
+	if (!PyGenUtil::ParseMethodParameters(InArgs, InKwds, InFuncDef.InputParams, InPythonFuncName, Params))
 	{
 		return nullptr;
 	}
 
-	if (InClass && InObj)
+	if (InObj && ensureAlways(InFuncDef.Func))
 	{
-		check(InObj->IsA(InClass) || InObj->GetClass()->ImplementsInterface(InClass));
-
-		// Find the function from the class rather than the object, as we may be specifically calling a super class function
-		UFunction* Func = InClass->FindFunctionByName(InFuncName);
-		if (ensureAlways(Func))
+		// Deprecated functions emit a warning
 		{
-			FStructOnScope FuncParams(Func);
-			CallFunction_ApplyDefaults(Func, FuncParams.GetStructMemory(), InParamDef);
-			for (int32 ParamIndex = 0; ParamIndex < Params.Num(); ++ParamIndex)
+			FString DeprecationMessage;
+			if (PyGenUtil::IsDeprecatedFunction(InFuncDef.Func, &DeprecationMessage) &&
+				PyUtil::SetPythonWarning(PyExc_DeprecationWarning, InErrorCtxt, *FString::Printf(TEXT("Function '%s.%s' is deprecated: %s"), *InFuncDef.Func->GetOwnerClass()->GetName(), *InFuncDef.Func->GetName(), *DeprecationMessage)) == -1
+				)
 			{
-				const PyGenUtil::FGeneratedWrappedMethodParameter& ParamDef = InParamDef[ParamIndex];
+				// -1 from SetPythonWarning means the warning should be an exception
+				return nullptr;
+			}
+		}
 
-				UProperty* ArgProp = FuncParams.GetStruct()->FindPropertyByName(ParamDef.ParamPropName);
-				if (!ArgProp)
-				{
-					PyUtil::SetPythonError(PyExc_Exception, InErrorCtxt, *FString::Printf(TEXT("Failed to find property '%s' for parameter '%s' when calling function '%s' on '%s'"), *ParamDef.ParamPropName.ToString(), UTF8_TO_TCHAR(ParamDef.ParamName.GetData()), *Func->GetName(), *InObj->GetName()));
-					return nullptr;
-				}
+		FStructOnScope FuncParams(InFuncDef.Func);
+		PyGenUtil::ApplyParamDefaults(FuncParams.GetStructMemory(), InFuncDef.InputParams);
+		for (int32 ParamIndex = 0; ParamIndex < Params.Num(); ++ParamIndex)
+		{
+			const PyGenUtil::FGeneratedWrappedMethodParameter& ParamDef = InFuncDef.InputParams[ParamIndex];
 
-				PyObject* PyValue = Params[ParamIndex];
-				if (PyValue)
+			PyObject* PyValue = Params[ParamIndex];
+			if (PyValue)
+			{
+				if (!PyConversion::NativizeProperty_InContainer(PyValue, ParamDef.ParamProp, FuncParams.GetStructMemory(), 0))
 				{
-					if (!PyConversion::NativizeProperty_InContainer(PyValue, ArgProp, FuncParams.GetStructMemory(), 0))
-					{
-						PyUtil::SetPythonError(PyExc_TypeError, InErrorCtxt, *FString::Printf(TEXT("Failed to convert parameter '%s' when calling function '%s' on '%s'"), UTF8_TO_TCHAR(ParamDef.ParamName.GetData()), *Func->GetName(), *InObj->GetName()));
-						return nullptr;
-					}
+					PyUtil::SetPythonError(PyExc_TypeError, InErrorCtxt, *FString::Printf(TEXT("Failed to convert parameter '%s' when calling function '%s.%s' on '%s'"), UTF8_TO_TCHAR(ParamDef.ParamName.GetData()), *InFuncDef.Func->GetOwnerClass()->GetName(), *InFuncDef.Func->GetName(), *InObj->GetName()));
+					return nullptr;
 				}
 			}
-			CallFunction_InvokeImpl(InObj, Func, FuncParams.GetStructMemory(), InErrorCtxt);
-			return CallFunction_ReturnImpl(InObj, Func, FuncParams.GetStructMemory(), InErrorCtxt);
 		}
+		PyUtil::InvokeFunctionCall(InObj, InFuncDef.Func, FuncParams.GetStructMemory(), InErrorCtxt);
+		return PyGenUtil::PackReturnValues(FuncParams.GetStructMemory(), InFuncDef.OutputParams, InErrorCtxt, *FString::Printf(TEXT("function '%s.%s' on '%s'"), *InFuncDef.Func->GetOwnerClass()->GetName(), *InFuncDef.Func->GetName(), *InObj->GetName()));
 	}
 
 	Py_RETURN_NONE;
 }
 
-void FPyWrapperObject::CallFunction_ApplyDefaults(UFunction* InFunc, void* InBaseParamsAddr, const TArray<PyGenUtil::FGeneratedWrappedMethodParameter>& InParamDef)
-{
-	for (const PyGenUtil::FGeneratedWrappedMethodParameter& ParamDef : InParamDef)
-	{
-		UProperty* ArgProp = InFunc->FindPropertyByName(ParamDef.ParamPropName);
-		if (ArgProp && ParamDef.ParamDefaultValue.IsSet())
-		{
-			ArgProp->ImportText(*ParamDef.ParamDefaultValue.GetValue(), ArgProp->ContainerPtrToValuePtr<void>(InBaseParamsAddr), PPF_None, nullptr);
-		}
-	}
-}
-
-void FPyWrapperObject::CallFunction_InvokeImpl(UObject* InObj, UFunction* InFunc, void* InBaseParamsAddr, const TCHAR* InErrorCtxt)
-{
-	FEditorScriptExecutionGuard ScriptGuard;
-	InObj->ProcessEvent(InFunc, InBaseParamsAddr);
-}
-
-PyObject* FPyWrapperObject::CallFunction_ReturnImpl(UObject* InObj, UFunction* InFunc, const void* InBaseParamsAddr, const TCHAR* InErrorCtxt)
-{
-	TArray<const UProperty*, TInlineAllocator<4>> OutParams;
-
-	// Add the return property first
-	if (const UProperty* ReturnProp = InFunc->GetReturnProperty())
-	{
-		OutParams.Add(ReturnProp);
-	}
-
-	// Add any other output params in order
-	for (TFieldIterator<const UProperty> ParamIt(InFunc); ParamIt; ++ParamIt)
-	{
-		const UProperty* Param = *ParamIt;
-		if (PyUtil::IsOutputParameter(Param))
-		{
-			OutParams.Add(Param);
-		}
-	}
-
-	return PyUtil::PackReturnValues(InFunc, OutParams, InBaseParamsAddr, InErrorCtxt, *FString::Printf(TEXT("function '%s' on '%s'"), *InFunc->GetName(), *InObj->GetName()));
-}
-
 PyObject* FPyWrapperObject::CallClassMethodNoArgs_Impl(PyTypeObject* InType, void* InClosure)
 {
 	const PyGenUtil::FGeneratedWrappedMethod* Closure = (PyGenUtil::FGeneratedWrappedMethod*)InClosure;
-	return CallFunction(Closure->Class, InType, Closure->MethodFuncName);
+	return CallFunction(InType, Closure->MethodFunc);
 }
 
 PyObject* FPyWrapperObject::CallClassMethodWithArgs_Impl(PyTypeObject* InType, PyObject* InArgs, PyObject* InKwds, void* InClosure)
 {
 	const PyGenUtil::FGeneratedWrappedMethod* Closure = (PyGenUtil::FGeneratedWrappedMethod*)InClosure;
-	return CallFunction(Closure->Class, InType, InArgs, InKwds, Closure->MethodFuncName, Closure->MethodName.GetData(), Closure->MethodParams);
+	return CallFunction(InType, InArgs, InKwds, Closure->MethodFunc, Closure->MethodName.GetData());
 }
 
 PyObject* FPyWrapperObject::CallMethodNoArgs_Impl(FPyWrapperObject* InSelf, void* InClosure)
 {
 	const PyGenUtil::FGeneratedWrappedMethod* Closure = (PyGenUtil::FGeneratedWrappedMethod*)InClosure;
-	return CallFunction(Closure->Class, InSelf, Closure->MethodFuncName);
+	return CallFunction(InSelf, Closure->MethodFunc);
 }
 
 PyObject* FPyWrapperObject::CallMethodWithArgs_Impl(FPyWrapperObject* InSelf, PyObject* InArgs, PyObject* InKwds, void* InClosure)
 {
 	const PyGenUtil::FGeneratedWrappedMethod* Closure = (PyGenUtil::FGeneratedWrappedMethod*)InClosure;
-	return CallFunction(Closure->Class, InSelf, InArgs, InKwds, Closure->MethodFuncName, Closure->MethodName.GetData(), Closure->MethodParams);
+	return CallFunction(InSelf, InArgs, InKwds, Closure->MethodFunc, Closure->MethodName.GetData());
 }
 
 PyObject* FPyWrapperObject::Getter_Impl(FPyWrapperObject* InSelf, void* InClosure)
 {
 	const PyGenUtil::FGeneratedWrappedGetSet* Closure = (PyGenUtil::FGeneratedWrappedGetSet*)InClosure;
-	return Closure->GetFuncName.IsNone()
-		? GetPropertyValue(InSelf, Closure->PropName, Closure->GetSetName.GetData())
-		: CallGetterFunction(Closure->Class, InSelf, Closure->GetFuncName);
+	return Closure->GetFunc.Func
+		? CallGetterFunction(InSelf, Closure->GetFunc)
+		: GetPropertyValue(InSelf, Closure->Prop, Closure->GetSetName.GetData());
 }
 
 int FPyWrapperObject::Setter_Impl(FPyWrapperObject* InSelf, PyObject* InValue, void* InClosure)
 {
 	const PyGenUtil::FGeneratedWrappedGetSet* Closure = (PyGenUtil::FGeneratedWrappedGetSet*)InClosure;
-	return Closure->SetFuncName.IsNone()
-		? SetPropertyValue(InSelf, InValue, Closure->PropName, Closure->GetSetName.GetData())
-		: CallSetterFunction(Closure->Class, InSelf, InValue, Closure->SetFuncName);
+	return Closure->SetFunc.Func
+		? CallSetterFunction(InSelf, InValue, Closure->SetFunc)
+		: SetPropertyValue(InSelf, InValue, Closure->Prop, Closure->GetSetName.GetData());
 }
 
 PyTypeObject InitializePyWrapperObjectType()
@@ -419,6 +427,18 @@ PyTypeObject InitializePyWrapperObjectType()
 			UClass* ObjClass = FPyWrapperObjectMetaData::GetClass(InSelf);
 			if (ObjClass)
 			{
+				// Deprecated classes emit a warning
+				{
+					FString DeprecationMessage;
+					if (PyGenUtil::IsDeprecatedClass(ObjClass, &DeprecationMessage) &&
+						PyUtil::SetPythonWarning(PyExc_DeprecationWarning, InSelf, *FString::Printf(TEXT("Class '%s' is deprecated: %s"), *ObjClass->GetName(), *DeprecationMessage)) == -1
+						)
+					{
+						// -1 from SetPythonWarning means the warning should be an exception
+						return -1;
+					}
+				}
+
 				if (ObjClass->HasAnyClassFlags(CLASS_Abstract))
 				{
 					PyUtil::SetPythonError(PyExc_Exception, InSelf, *FString::Printf(TEXT("Class '%s' is abstract"), *ObjClass->GetName()));
@@ -490,6 +510,13 @@ PyTypeObject InitializePyWrapperObjectType()
 			return nullptr;
 		}
 
+		static PyObject* GetDefaultObject(PyTypeObject* InType)
+		{
+			UClass* Class = FPyWrapperObjectMetaData::GetClass(InType);
+			UObject* CDO = Class ? Class->GetDefaultObject() : nullptr;
+			return PyConversion::Pythonize(CDO);
+		}
+
 		static PyObject* StaticClass(PyTypeObject* InType)
 		{
 			UClass* Class = FPyWrapperObjectMetaData::GetClass(InType);
@@ -674,7 +701,7 @@ PyTypeObject InitializePyWrapperObjectType()
 			}
 
 			const FName ResolvedName = FPyWrapperObjectMetaData::ResolvePropertyName(InSelf, Name);
-			return FPyWrapperObject::GetPropertyValue(InSelf, ResolvedName, TCHAR_TO_UTF8(*Name.ToString()));
+			return FPyWrapperObject::GetPropertyValueByName(InSelf, ResolvedName, TCHAR_TO_UTF8(*Name.ToString()));
 		}
 
 		static PyObject* SetEditorProperty(FPyWrapperObject* InSelf, PyObject* InArgs, PyObject* InKwds)
@@ -696,7 +723,7 @@ PyTypeObject InitializePyWrapperObjectType()
 			}
 
 			const FName ResolvedName = FPyWrapperObjectMetaData::ResolvePropertyName(InSelf, Name);
-			const int Result = FPyWrapperObject::SetPropertyValue(InSelf, PyValueObj, ResolvedName, TCHAR_TO_UTF8(*Name.ToString()), /*InNotifyChange*/true, CPF_EditConst);
+			const int Result = FPyWrapperObject::SetPropertyValueByName(InSelf, PyValueObj, ResolvedName, TCHAR_TO_UTF8(*Name.ToString()), /*InNotifyChange*/true, CPF_EditConst);
 			if (Result != 0)
 			{
 				return nullptr;
@@ -709,6 +736,7 @@ PyTypeObject InitializePyWrapperObjectType()
 	static PyMethodDef PyMethods[] = {
 		{ PyGenUtil::PostInitFuncName, PyCFunctionCast(&FMethods::PostInit), METH_NOARGS, "x._post_init() -- called during Unreal object initialization (equivalent to PostInitProperties in C++)" },
 		{ "cast", PyCFunctionCast(&FMethods::Cast), METH_VARARGS | METH_CLASS, "X.cast(object) -> UObject -- cast the given object to this Unreal object type" },
+		{ "get_default_object", PyCFunctionCast(&FMethods::GetDefaultObject), METH_NOARGS | METH_CLASS, "X.get_default_object() -> UObject -- get the Unreal class default object (CDO) of this type" },
 		{ "static_class", PyCFunctionCast(&FMethods::StaticClass), METH_NOARGS | METH_CLASS, "X.static_class() -> UClass -- get the Unreal class of this type" },
 		{ "get_class", PyCFunctionCast(&FMethods::GetClass), METH_NOARGS, "x.get_class() -> UClass -- get the Unreal class of this instance" },
 		{ "get_outer", PyCFunctionCast(&FMethods::GetOuter), METH_NOARGS, "x.get_outer() -> UObject -- get the outer object from this instance (if any)" },
@@ -749,29 +777,25 @@ PyTypeObject InitializePyWrapperObjectType()
 
 PyTypeObject PyWrapperObjectType = InitializePyWrapperObjectType();
 
-FPyWrapperObjectMetaData::FPyWrapperObjectMetaData()
-	: Class(nullptr)
+void FPyWrapperObjectMetaData::AddReferencedObjects(FPyWrapperBase* Instance, FReferenceCollector& Collector)
 {
-	AddReferencedObjects = [](FPyWrapperBase* Self, FReferenceCollector& Collector)
-	{
-		FPyWrapperObject* ObjSelf = static_cast<FPyWrapperObject*>(Self);
+	FPyWrapperObject* Self = static_cast<FPyWrapperObject*>(Instance);
 
-		UObject* OldInstance = ObjSelf->ObjectInstance;
-		Collector.AddReferencedObject(ObjSelf->ObjectInstance);
-		if (ObjSelf->ObjectInstance != OldInstance && ObjSelf->ObjectInstance != nullptr)
+	UObject* OldInstance = Self->ObjectInstance;
+	Collector.AddReferencedObject(Self->ObjectInstance);
+	if (Self->ObjectInstance != OldInstance && Self->ObjectInstance != nullptr)
+	{
+		// Object instance has been re-pointed, make sure we're still the correct type
+		PyTypeObject* NewPyType = FPyWrapperTypeRegistry::Get().GetWrappedClassType(Self->ObjectInstance->GetClass());
+		if (PyType_IsSubtype(NewPyType, &PyWrapperObjectType) && NewPyType->tp_basicsize == Py_TYPE(Self)->tp_basicsize)
 		{
-			// Object instance has been re-pointed, make sure we're still the correct type
-			PyTypeObject* NewPyType = FPyWrapperTypeRegistry::Get().GetWrappedClassType(ObjSelf->ObjectInstance->GetClass());
-			if (PyType_IsSubtype(NewPyType, &PyWrapperObjectType) && NewPyType->tp_basicsize == Py_TYPE(ObjSelf)->tp_basicsize)
-			{
-				Py_TYPE(ObjSelf) = NewPyType; // todo: is this safe?
-			}
-			else
-			{
-				ObjSelf->ObjectInstance = nullptr;
-			}
+			Py_TYPE(Self) = NewPyType; // todo: is this safe?
 		}
-	};
+		else
+		{
+			Self->ObjectInstance = nullptr;
+		}
+	}
 }
 
 UClass* FPyWrapperObjectMetaData::GetClass(PyTypeObject* PyType)
@@ -904,25 +928,23 @@ struct FPythonGeneratedClassUtil
 		}
 
 		// Build the definition data for the new property accessor
-		PyGenUtil::FPropertyDef& PropDef = *InClass->PropertyDefs[InClass->PropertyDefs.Add(MakeShared<PyGenUtil::FPropertyDef>())];
-		PropDef.GeneratedWrappedGetSet.Class = InClass;
+		PyGenUtil::FPropertyDef& PropDef = *InClass->PropertyDefs.Add_GetRef(MakeShared<PyGenUtil::FPropertyDef>());
 		PropDef.GeneratedWrappedGetSet.GetSetName = PyGenUtil::TCHARToUTF8Buffer(*InFieldName);
 		PropDef.GeneratedWrappedGetSet.GetSetDoc = PyGenUtil::TCHARToUTF8Buffer(*FString::Printf(TEXT("type: %s\n%s"), *PyGenUtil::GetPropertyPythonType(Prop), *PyGenUtil::GetFieldTooltip(Prop)));
-		PropDef.GeneratedWrappedGetSet.PropName = Prop->GetFName();
-		PropDef.GeneratedWrappedGetSet.GetFuncName = GetterFuncName;
-		PropDef.GeneratedWrappedGetSet.SetFuncName = SetterFuncName;
+		PropDef.GeneratedWrappedGetSet.Prop = Prop;
+		PropDef.GeneratedWrappedGetSet.GetFunc.SetFunctionAndExtractParams(InClass->FindFunctionByName(GetterFuncName));
+		PropDef.GeneratedWrappedGetSet.SetFunc.SetFunctionAndExtractParams(InClass->FindFunctionByName(SetterFuncName));
 		PropDef.GeneratedWrappedGetSet.GetCallback = (getter)&FPyWrapperObject::Getter_Impl;
 		PropDef.GeneratedWrappedGetSet.SetCallback = (setter)&FPyWrapperObject::Setter_Impl;
 		PropDef.GeneratedWrappedGetSet.ToPython(PropDef.PyGetSet);
 
 		// If this property has a getter or setter, also make an internal version with the get/set function cleared so that Python can read/write the internal property value
-		if (!GetterFuncName.IsNone() || !SetterFuncName.IsNone())
+		if (PropDef.GeneratedWrappedGetSet.GetFunc.Func || PropDef.GeneratedWrappedGetSet.SetFunc.Func)
 		{
-			PyGenUtil::FPropertyDef& InternalPropDef = *InClass->PropertyDefs[InClass->PropertyDefs.Add(MakeShared<PyGenUtil::FPropertyDef>())];
-			InternalPropDef.GeneratedWrappedGetSet.Class = InClass;
+			PyGenUtil::FPropertyDef& InternalPropDef = *InClass->PropertyDefs.Add_GetRef(MakeShared<PyGenUtil::FPropertyDef>());
 			InternalPropDef.GeneratedWrappedGetSet.GetSetName = PyGenUtil::TCHARToUTF8Buffer(*FString::Printf(TEXT("_%s"), *InFieldName));
 			InternalPropDef.GeneratedWrappedGetSet.GetSetDoc = PropDef.GeneratedWrappedGetSet.GetSetDoc;
-			InternalPropDef.GeneratedWrappedGetSet.PropName = Prop->GetFName();
+			InternalPropDef.GeneratedWrappedGetSet.Prop = Prop;
 			InternalPropDef.GeneratedWrappedGetSet.GetCallback = (getter)&FPyWrapperObject::Getter_Impl;
 			InternalPropDef.GeneratedWrappedGetSet.SetCallback = (setter)&FPyWrapperObject::Setter_Impl;
 			InternalPropDef.GeneratedWrappedGetSet.ToPython(InternalPropDef.PyGetSet);
@@ -1093,66 +1115,59 @@ struct FPythonGeneratedClassUtil
 			}
 		}
 		// Apply the defaults to the function arguments and build the Python method params
-		TArray<PyGenUtil::FGeneratedWrappedMethodParameter> MethodParams;
+		PyGenUtil::FGeneratedWrappedFunction GeneratedWrappedFunction;
+		GeneratedWrappedFunction.SetFunctionAndExtractParams(Func);
+		// SetFunctionAndExtractParams doesn't always use the correct names or defaults for generated classes
+		for (int32 InputArgIndex = 0; InputArgIndex < GeneratedWrappedFunction.InputParams.Num(); ++InputArgIndex)
 		{
-			int32 InputArgIndex = 0;
-			for (TFieldIterator<const UProperty> ParamIt(Func); ParamIt; ++ParamIt)
-			{
-				const UProperty* Param = *ParamIt;
+			PyGenUtil::FGeneratedWrappedMethodParameter& GeneratedWrappedMethodParam = GeneratedWrappedFunction.InputParams[InputArgIndex];
+			const UProperty* Param = GeneratedWrappedMethodParam.ParamProp;
 
-				if (PyUtil::IsInputParameter(Param))
-				{
-					const FName DefaultValueMetaDataKey = *FString::Printf(TEXT("CPP_Default_%s"), *Param->GetName());
+			const FName DefaultValueMetaDataKey = *FString::Printf(TEXT("CPP_Default_%s"), *Param->GetName());
 
-					TOptional<FString> ResolvedDefaultValue;
-					if (FuncArgDefaults.IsValidIndex(InputArgIndex) && FuncArgDefaults[InputArgIndex])
-					{
-						// Convert the default value to the given property...
-						PyUtil::FPropValueOnScope DefaultValue(Param);
-						if (!DefaultValue.IsValid() || !DefaultValue.SetValue(FuncArgDefaults[InputArgIndex], *PyUtil::GetErrorContext(InPyType)))
-						{
-							PyUtil::SetPythonError(PyExc_Exception, InPyType, *FString::Printf(TEXT("Failed to convert default value for function '%s' argument '%s' (%s)"), *InFieldName, *FuncArgNames[InputArgIndex], *Param->GetClass()->GetName()));
-							return false;
-						}
-
-						// ... and export it as meta-data
-						FString ExportedDefaultValue;
-						if (!DefaultValue.GetProp()->ExportText_Direct(ExportedDefaultValue, DefaultValue.GetValue(), DefaultValue.GetValue(), nullptr, PPF_None))
-						{
-							PyUtil::SetPythonError(PyExc_Exception, InPyType, *FString::Printf(TEXT("Failed to export default value for function '%s' argument '%s' (%s)"), *InFieldName, *FuncArgNames[InputArgIndex], *Param->GetClass()->GetName()));
-							return false;
-						}
-
-						ResolvedDefaultValue = ExportedDefaultValue;
-					}
-					if (!ResolvedDefaultValue.IsSet() && SuperFunc && SuperFunc->HasAnyFunctionFlags(FUNC_HasDefaults))
-					{
-						if (SuperFunc->HasMetaData(DefaultValueMetaDataKey))
-						{
-							ResolvedDefaultValue = SuperFunc->GetMetaData(DefaultValueMetaDataKey);
-						}
-					}
-					if (ResolvedDefaultValue.IsSet())
-					{
-						Func->SetMetaData(DefaultValueMetaDataKey, *ResolvedDefaultValue.GetValue());
-						Func->FunctionFlags |= FUNC_HasDefaults;
-					}
+			TOptional<FString> ResolvedDefaultValue;
+			if (FuncArgDefaults.IsValidIndex(InputArgIndex) && FuncArgDefaults[InputArgIndex])
+			{
+				// Convert the default value to the given property...
+				PyUtil::FPropValueOnScope DefaultValue(Param);
+				if (!DefaultValue.IsValid() || !DefaultValue.SetValue(FuncArgDefaults[InputArgIndex], *PyUtil::GetErrorContext(InPyType)))
+				{
+					PyUtil::SetPythonError(PyExc_Exception, InPyType, *FString::Printf(TEXT("Failed to convert default value for function '%s' argument '%s' (%s)"), *InFieldName, *FuncArgNames[InputArgIndex], *Param->GetClass()->GetName()));
+					return false;
+				}
 
-					PyGenUtil::FGeneratedWrappedMethodParameter& GeneratedWrappedMethodParam = MethodParams[MethodParams.AddDefaulted()];
-					GeneratedWrappedMethodParam.ParamName = PyGenUtil::TCHARToUTF8Buffer(FuncArgNames.IsValidIndex(InputArgIndex) ? *FuncArgNames[InputArgIndex] : *Param->GetName());
-					GeneratedWrappedMethodParam.ParamPropName = Param->GetFName();
-					GeneratedWrappedMethodParam.ParamDefaultValue = MoveTemp(ResolvedDefaultValue);
+				// ... and export it as meta-data
+				FString ExportedDefaultValue;
+				if (!DefaultValue.GetProp()->ExportText_Direct(ExportedDefaultValue, DefaultValue.GetValue(), DefaultValue.GetValue(), nullptr, PPF_None))
+				{
+					PyUtil::SetPythonError(PyExc_Exception, InPyType, *FString::Printf(TEXT("Failed to export default value for function '%s' argument '%s' (%s)"), *InFieldName, *FuncArgNames[InputArgIndex], *Param->GetClass()->GetName()));
+					return false;
+				}
 
-					++InputArgIndex;
+				ResolvedDefaultValue = ExportedDefaultValue;
+			}
+			if (!ResolvedDefaultValue.IsSet() && SuperFunc && SuperFunc->HasAnyFunctionFlags(FUNC_HasDefaults))
+			{
+				if (SuperFunc->HasMetaData(DefaultValueMetaDataKey))
+				{
+					ResolvedDefaultValue = SuperFunc->GetMetaData(DefaultValueMetaDataKey);
 				}
 			}
+			if (ResolvedDefaultValue.IsSet())
+			{
+				Func->SetMetaData(DefaultValueMetaDataKey, *ResolvedDefaultValue.GetValue());
+				Func->FunctionFlags |= FUNC_HasDefaults;
+			}
+
+			GeneratedWrappedMethodParam.ParamName = PyGenUtil::TCHARToUTF8Buffer(FuncArgNames.IsValidIndex(InputArgIndex) ? *FuncArgNames[InputArgIndex] : *Param->GetName());
+			GeneratedWrappedMethodParam.ParamDefaultValue = MoveTemp(ResolvedDefaultValue);
 		}
 		Func->Bind();
 		Func->StaticLink(true);
 
-		if (MethodParams.Num() != FuncArgNames.Num())
+		if (GeneratedWrappedFunction.InputParams.Num() != FuncArgNames.Num())
 		{
-			PyUtil::SetPythonError(PyExc_Exception, InPyType, *FString::Printf(TEXT("Incorrect number of arguments specified for '%s' (expected %d, got %d)"), *InFieldName, MethodParams.Num(), FuncArgNames.Num()));
+			PyUtil::SetPythonError(PyExc_Exception, InPyType, *FString::Printf(TEXT("Incorrect number of arguments specified for '%s' (expected %d, got %d)"), *InFieldName, GeneratedWrappedFunction.InputParams.Num(), FuncArgNames.Num()));
 			return false;
 		}
 
@@ -1168,11 +1183,10 @@ struct FPythonGeneratedClassUtil
 		}
 
 		// Build the definition data for the new method
-		PyGenUtil::FFunctionDef& FuncDef = *InClass->FunctionDefs[InClass->FunctionDefs.Add(MakeShared<PyGenUtil::FFunctionDef>())];
-		FuncDef.GeneratedWrappedMethod.Class = InClass;
+		PyGenUtil::FFunctionDef& FuncDef = *InClass->FunctionDefs.Add_GetRef(MakeShared<PyGenUtil::FFunctionDef>());
 		FuncDef.GeneratedWrappedMethod.MethodName = PyGenUtil::TCHARToUTF8Buffer(*InFieldName);
 		FuncDef.GeneratedWrappedMethod.MethodDoc = PyGenUtil::TCHARToUTF8Buffer(*PyGenUtil::GetFieldTooltip(Func));
-		FuncDef.GeneratedWrappedMethod.MethodFuncName = Func->GetFName();
+		FuncDef.GeneratedWrappedMethod.MethodFunc = MoveTemp(GeneratedWrappedFunction);
 		FuncDef.GeneratedWrappedMethod.MethodFlags = FuncArgNames.Num() > 0 ? METH_VARARGS | METH_KEYWORDS : METH_NOARGS;
 		if (Func->HasAnyFunctionFlags(FUNC_Static))
 		{
@@ -1183,7 +1197,6 @@ struct FPythonGeneratedClassUtil
 		{
 			FuncDef.GeneratedWrappedMethod.MethodCallback = FuncArgNames.Num() > 0 ? PyCFunctionWithClosureCast(&FPyWrapperObject::CallMethodWithArgs_Impl) : PyCFunctionWithClosureCast(&FPyWrapperObject::CallMethodNoArgs_Impl);
 		}
-		FuncDef.GeneratedWrappedMethod.MethodParams = MoveTemp(MethodParams);
 		FuncDef.GeneratedWrappedMethod.ToPython(FuncDef.PyMethod);
 		FuncDef.PyFunction = FPyObjectPtr::NewReference(InPyFuncDef->Func);
 		FuncDef.bIsHidden = EnumHasAnyFlags(InPyFuncDef->FuncFlags, EPyUFunctionDefFlags::Getter | EPyUFunctionDefFlags::Setter);
@@ -1196,27 +1209,31 @@ struct FPythonGeneratedClassUtil
 		InClass->PropertyDefs.Reserve(InOldClass->PropertyDefs.Num());
 		for (const TSharedPtr<PyGenUtil::FPropertyDef>& OldPropDef : InOldClass->PropertyDefs)
 		{
-			const FName PropName = OldPropDef->GeneratedWrappedGetSet.PropName;
+			const UProperty* OldProp = OldPropDef->GeneratedWrappedGetSet.Prop;
+			const UFunction* OldGetter = OldPropDef->GeneratedWrappedGetSet.GetFunc.Func;
+			const UFunction* OldSetter = OldPropDef->GeneratedWrappedGetSet.SetFunc.Func;
 
-			UProperty* OldProp = InOldClass->FindPropertyByName(PropName);
-			if (!OldProp)
-			{
-				PyUtil::SetPythonError(PyExc_Exception, InPyType, *FString::Printf(TEXT("Failed to find property for '%s'"), UTF8_TO_TCHAR(OldPropDef->PyGetSet.name)));
-				return false;
-			}
-
-			UProperty* Prop = DuplicateObject<UProperty>(OldProp, InClass, PropName);
+			UProperty* Prop = DuplicateObject<UProperty>(OldProp, InClass, OldProp->GetFName());
 			if (!Prop)
 			{
 				PyUtil::SetPythonError(PyExc_Exception, InPyType, *FString::Printf(TEXT("Failed to duplicate property for '%s'"), UTF8_TO_TCHAR(OldPropDef->PyGetSet.name)));
 				return false;
 			}
 
-			UMetaData::CopyMetadata(OldProp, Prop);
+			UMetaData::CopyMetadata((UProperty*)OldProp, Prop);
 			InClass->AddCppProperty(Prop);
 
-			PyGenUtil::FPropertyDef& PropDef = *InClass->PropertyDefs[InClass->PropertyDefs.Add(MakeShared<PyGenUtil::FPropertyDef>(*OldPropDef))];
-			PropDef.GeneratedWrappedGetSet.Class = InClass;
+			PyGenUtil::FPropertyDef& PropDef = *InClass->PropertyDefs.Add_GetRef(MakeShared<PyGenUtil::FPropertyDef>());
+			PropDef.GeneratedWrappedGetSet = OldPropDef->GeneratedWrappedGetSet;
+			PropDef.GeneratedWrappedGetSet.Prop = Prop;
+			if (OldGetter)
+			{
+				PropDef.GeneratedWrappedGetSet.GetFunc.SetFunctionAndExtractParams(InClass->FindFunctionByName(OldGetter->GetFName()));
+			}
+			if (OldSetter)
+			{
+				PropDef.GeneratedWrappedGetSet.SetFunc.SetFunctionAndExtractParams(InClass->FindFunctionByName(OldSetter->GetFName()));
+			}
 			PropDef.GeneratedWrappedGetSet.ToPython(PropDef.PyGetSet);
 		}
 
@@ -1228,31 +1245,27 @@ struct FPythonGeneratedClassUtil
 		InClass->FunctionDefs.Reserve(InOldClass->FunctionDefs.Num());
 		for (const TSharedPtr<PyGenUtil::FFunctionDef>& OldFuncDef : InOldClass->FunctionDefs)
 		{
-			const FName FuncName = OldFuncDef->GeneratedWrappedMethod.MethodFuncName;
-
-			UFunction* OldFunc = InOldClass->FindFunctionByName(FuncName);
-			if (!OldFunc)
-			{
-				PyUtil::SetPythonError(PyExc_Exception, InPyType, *FString::Printf(TEXT("Failed to find function for '%s'"), UTF8_TO_TCHAR(OldFuncDef->PyMethod.MethodName)));
-				return false;
-			}
+			const UFunction* OldFunc = OldFuncDef->GeneratedWrappedMethod.MethodFunc.Func;
 
-			InClass->AddNativeFunction(*FuncName.ToString(), &UPythonGeneratedClass::CallPythonFunction);
-			UFunction* Func = DuplicateObject<UFunction>(OldFunc, InClass, FuncName);
+			InClass->AddNativeFunction(*OldFunc->GetName(), &UPythonGeneratedClass::CallPythonFunction);
+			UFunction* Func = DuplicateObject<UFunction>(OldFunc, InClass, OldFunc->GetFName());
 			if (!Func)
 			{
 				PyUtil::SetPythonError(PyExc_Exception, InPyType, *FString::Printf(TEXT("Failed to duplicate function for '%s'"), UTF8_TO_TCHAR(OldFuncDef->PyMethod.MethodName)));
 				return false;
 			}
 
-			UMetaData::CopyMetadata(OldFunc, Func);
+			UMetaData::CopyMetadata((UFunction*)OldFunc, Func);
 			InClass->AddFunctionToFunctionMap(Func, Func->GetFName());
 
 			Func->Bind();
 			Func->StaticLink(true);
 
-			PyGenUtil::FFunctionDef& FuncDef = *InClass->FunctionDefs[InClass->FunctionDefs.Add(MakeShared<PyGenUtil::FFunctionDef>(*OldFuncDef))];
-			FuncDef.GeneratedWrappedMethod.Class = InClass;
+			PyGenUtil::FFunctionDef& FuncDef = *InClass->FunctionDefs.Add_GetRef(MakeShared<PyGenUtil::FFunctionDef>());
+			FuncDef.GeneratedWrappedMethod = OldFuncDef->GeneratedWrappedMethod;
+			FuncDef.GeneratedWrappedMethod.MethodFunc.SetFunctionAndExtractParams(Func);
+			FuncDef.PyFunction = OldFuncDef->PyFunction;
+			FuncDef.bIsHidden = OldFuncDef->bIsHidden;
 			FuncDef.GeneratedWrappedMethod.ToPython(FuncDef.PyMethod);
 		}
 
@@ -1382,7 +1395,8 @@ UPythonGeneratedClass* UPythonGeneratedClass::GenerateClass(PyTypeObject* InPyTy
 		return nullptr;
 	}
 
-	// Add the fields to this class
+	// Add the functions to this class
+	// We have to process these first as properties may reference them as get/set functions
 	{
 		PyObject* FieldKey = nullptr;
 		PyObject* FieldValue = nullptr;
@@ -1398,19 +1412,30 @@ UPythonGeneratedClass* UPythonGeneratedClass::GenerateClass(PyTypeObject* InPyTy
 				return nullptr;
 			}
 
-			if (PyObject_IsInstance(FieldValue, (PyObject*)&PyUPropertyDefType) == 1)
+			if (PyObject_IsInstance(FieldValue, (PyObject*)&PyUFunctionDefType) == 1)
 			{
-				FPyUPropertyDef* PyPropDef = (FPyUPropertyDef*)FieldValue;
-				if (!FPythonGeneratedClassUtil::CreatePropertyFromDefinition(Class, InPyType, FieldName, PyPropDef))
+				FPyUFunctionDef* PyFuncDef = (FPyUFunctionDef*)FieldValue;
+				if (!FPythonGeneratedClassUtil::CreateFunctionFromDefinition(Class, InPyType, FieldName, PyFuncDef))
 				{
 					return nullptr;
 				}
 			}
+		}
+	}
 
-			if (PyObject_IsInstance(FieldValue, (PyObject*)&PyUFunctionDefType) == 1)
+	// Add the properties to this class
+	{
+		PyObject* FieldKey = nullptr;
+		PyObject* FieldValue = nullptr;
+		Py_ssize_t FieldIndex = 0;
+		while (PyDict_Next(InPyType->tp_dict, &FieldIndex, &FieldKey, &FieldValue))
+		{
+			const FString FieldName = PyUtil::PyObjectToUEString(FieldKey);
+
+			if (PyObject_IsInstance(FieldValue, (PyObject*)&PyUPropertyDefType) == 1)
 			{
-				FPyUFunctionDef* PyFuncDef = (FPyUFunctionDef*)FieldValue;
-				if (!FPythonGeneratedClassUtil::CreateFunctionFromDefinition(Class, InPyType, FieldName, PyFuncDef))
+				FPyUPropertyDef* PyPropDef = (FPyUPropertyDef*)FieldValue;
+				if (!FPythonGeneratedClassUtil::CreatePropertyFromDefinition(Class, InPyType, FieldName, PyPropDef))
 				{
 					return nullptr;
 				}
@@ -1476,11 +1501,11 @@ UPythonGeneratedClass* UPythonGeneratedClass::ReparentClass(UPythonGeneratedClas
 
 	// Copy the data from the old class
 	Class->PyPostInitFunction = InOldClass->PyPostInitFunction;
-	if (!FPythonGeneratedClassUtil::CopyPropertiesFromOldClass(Class, InOldClass, PyType))
+	if (!FPythonGeneratedClassUtil::CopyFunctionsFromOldClass(Class, InOldClass, PyType))
 	{
 		return nullptr;
 	}
-	if (!FPythonGeneratedClassUtil::CopyFunctionsFromOldClass(Class, InOldClass, PyType))
+	if (!FPythonGeneratedClassUtil::CopyPropertiesFromOldClass(Class, InOldClass, PyType))
 	{
 		return nullptr;
 	}
@@ -1513,16 +1538,15 @@ DEFINE_FUNCTION(UPythonGeneratedClass::CallPythonFunction)
 	const UPythonGeneratedClass* This = CastChecked<UPythonGeneratedClass>(Stack.Node->GetOwnerClass());
 
 	// Find the Python function to call
-	PyObject* PyFunc = nullptr;
-	for (const TSharedPtr<PyGenUtil::FFunctionDef>& FuncDef : This->FunctionDefs)
+	TSharedPtr<PyGenUtil::FFunctionDef> FuncDef;
 	{
-		if (FuncDef->GeneratedWrappedMethod.MethodFuncName == Stack.Node->GetFName())
+		const TSharedPtr<PyGenUtil::FFunctionDef>* FuncDefPtr = This->FunctionDefs.FindByPredicate([&Stack](const TSharedPtr<PyGenUtil::FFunctionDef>& InFuncDef)
 		{
-			PyFunc = FuncDef->PyFunction;
-			break;
-		}
+			return InFuncDef->GeneratedWrappedMethod.MethodFunc.Func == Stack.Node;
+		});
+		FuncDef = FuncDefPtr ? *FuncDefPtr : nullptr;
 	}
-	if (!PyFunc)
+	if (!FuncDef.IsValid())
 	{
 		UE_LOG(LogPython, Error, TEXT("Failed to find Python function for '%s' on '%s'"), *Stack.Node->GetName(), *This->GetName());
 	}
@@ -1550,7 +1574,7 @@ DEFINE_FUNCTION(UPythonGeneratedClass::CallPythonFunction)
 				PyArgs = FPyObjectPtr::StealReference(PyTuple_New(1));
 				PyTuple_SetItem(PyArgs, 0, PySelf.Release()); // SetItem steals the reference
 			}
-			FPyObjectPtr RetVals = FPyObjectPtr::StealReference(PyObject_CallObject(PyFunc, PyArgs));
+			FPyObjectPtr RetVals = FPyObjectPtr::StealReference(PyObject_CallObject(FuncDef->PyFunction, PyArgs));
 			if (!RetVals)
 			{
 				return false;
@@ -1560,37 +1584,20 @@ DEFINE_FUNCTION(UPythonGeneratedClass::CallPythonFunction)
 		{
 			// Complex case, parameters or return value
 			TArray<FPyObjectPtr, TInlineAllocator<4>> PyParams;
-			TArray<const UProperty*, TInlineAllocator<4>> OutParams;
-
-			// Add the return property first
-			const UProperty* ReturnProp = Stack.Node->GetReturnProperty();
-			if (ReturnProp)
-			{
-				OutParams.Add(ReturnProp);
-			}
 
-			// Add any other output params in order, and get the value of the input params for the Python args
-			int32 ArgIndex = 0;
-			for (TFieldIterator<const UProperty> ParamIt(Stack.Node); ParamIt; ++ParamIt)
+			// Get the value of the input params for the Python args
 			{
-				const UProperty* Param = *ParamIt;
-
-				if (PyUtil::IsInputParameter(Param))
+				int32 ArgIndex = 0;
+				for (const PyGenUtil::FGeneratedWrappedMethodParameter& ParamDef : FuncDef->GeneratedWrappedMethod.MethodFunc.InputParams)
 				{
-					FPyObjectPtr& PyParam = PyParams[PyParams.AddDefaulted()];
-					if (!PyConversion::PythonizeProperty_InContainer(Param, Stack.Locals, 0, PyParam.Get()))
+					FPyObjectPtr& PyParam = PyParams.AddDefaulted_GetRef();
+					if (!PyConversion::PythonizeProperty_InContainer(ParamDef.ParamProp, Stack.Locals, 0, PyParam.Get()))
 					{
-						PyUtil::SetPythonError(PyExc_TypeError, PyFunc, *FString::Printf(TEXT("Failed to convert argument at pos '%d' when calling function '%s' on '%s'"), ArgIndex + 1, *Stack.Node->GetName(), *P_THIS_OBJECT->GetName()));
+						PyUtil::SetPythonError(PyExc_TypeError, FuncDef->PyFunction, *FString::Printf(TEXT("Failed to convert argument at pos '%d' when calling function '%s' on '%s'"), ArgIndex + 1, *Stack.Node->GetName(), *P_THIS_OBJECT->GetName()));
 						return false;
 					}
-
 					++ArgIndex;
 				}
-
-				if (PyUtil::IsOutputParameter(Param))
-				{
-					OutParams.Add(Param);
-				}
 			}
 
 			const int32 PyParamOffset = (PySelf ? 1 : 0);
@@ -1604,25 +1611,25 @@ DEFINE_FUNCTION(UPythonGeneratedClass::CallPythonFunction)
 				PyTuple_SetItem(PyArgs, PyParamIndex + PyParamOffset, PyParams[PyParamIndex].Release()); // SetItem steals the reference
 			}
 
-			FPyObjectPtr RetVals = FPyObjectPtr::StealReference(PyObject_CallObject(PyFunc, PyArgs));
+			FPyObjectPtr RetVals = FPyObjectPtr::StealReference(PyObject_CallObject(FuncDef->PyFunction, PyArgs));
 			if (!RetVals)
 			{
 				return false;
 			}
 
-			if (!PyUtil::UnpackReturnValues(RetVals, Stack.Node, OutParams, Stack.Locals, *PyUtil::GetErrorContext(PyFunc), *FString::Printf(TEXT("function '%s' on '%s'"), *Stack.Node->GetName(), *P_THIS_OBJECT->GetName())))
+			if (!PyGenUtil::UnpackReturnValues(RetVals, Stack.Locals, FuncDef->GeneratedWrappedMethod.MethodFunc.OutputParams, *PyUtil::GetErrorContext(FuncDef->PyFunction), *FString::Printf(TEXT("function '%s' on '%s'"), *Stack.Node->GetName(), *P_THIS_OBJECT->GetName())))
 			{
 				return false;
 			}
 
-			if (ReturnProp)
+			// Copy the data back out of the function call
+			if (const UProperty* ReturnProp = Stack.Node->GetReturnProperty())
 			{
 				ReturnProp->CopyCompleteValue(RESULT_PARAM, ReturnProp->ContainerPtrToValuePtr<void>(Stack.Locals));
-				
-				for (FOutParmRec* OutParamRec = Stack.OutParms; OutParamRec; OutParamRec = OutParamRec->NextOutParm)
-				{
-					ReturnProp->CopyCompleteValue(OutParamRec->PropAddr, OutParamRec->Property->ContainerPtrToValuePtr<void>(Stack.Locals));
-				}
+			}
+			for (FOutParmRec* OutParamRec = Stack.OutParms; OutParamRec; OutParamRec = OutParamRec->NextOutParm)
+			{
+				OutParamRec->Property->CopyCompleteValue(OutParamRec->PropAddr, OutParamRec->Property->ContainerPtrToValuePtr<void>(Stack.Locals));
 			}
 		}
 
diff --git a/Engine/Plugins/Experimental/PythonScriptPlugin/Source/PythonScriptPlugin/Private/PyWrapperObject.h b/Engine/Plugins/Experimental/PythonScriptPlugin/Source/PythonScriptPlugin/Private/PyWrapperObject.h
--- a/Engine/Plugins/Experimental/PythonScriptPlugin/Source/PythonScriptPlugin/Private/PyWrapperObject.h
+++ b/Engine/Plugins/Experimental/PythonScriptPlugin/Source/PythonScriptPlugin/Private/PyWrapperObject.h
@@ -38,49 +38,46 @@ struct FPyWrapperObject : public FPyWrapperBase
 	static bool ValidateInternalState(FPyWrapperObject* InSelf);
 
 	/** Cast the given Python object to this wrapped type (returns a new reference) */
-	static FPyWrapperObject* CastPyObject(PyObject* InPyObject);
+	static FPyWrapperObject* CastPyObject(PyObject* InPyObject, FPyConversionResult* OutCastResult = nullptr);
 
 	/** Cast the given Python object to this wrapped type, or attempt to convert the type into a new wrapped instance (returns a new reference) */
-	static FPyWrapperObject* CastPyObject(PyObject* InPyObject, PyTypeObject* InType);
+	static FPyWrapperObject* CastPyObject(PyObject* InPyObject, PyTypeObject* InType, FPyConversionResult* OutCastResult = nullptr);
 
 	/** Get a named property value from this instance (called via generated code) */
-	static PyObject* GetPropertyValue(FPyWrapperObject* InSelf, const FName InPropName, const char* InPythonAttrName);
+	static PyObject* GetPropertyValueByName(FPyWrapperObject* InSelf, const FName InPropName, const char* InPythonAttrName);
+
+	/** Get a property value from this instance (called via generated code) */
+	static PyObject* GetPropertyValue(FPyWrapperObject* InSelf, const UProperty* InProp, const char* InPythonAttrName);
 
 	/** Set a named property value on this instance (called via generated code) */
-	static int SetPropertyValue(FPyWrapperObject* InSelf, PyObject* InValue, const FName InPropName, const char* InPythonAttrName, const bool InNotifyChange = false, const uint64 InReadOnlyFlags = CPF_EditConst | CPF_BlueprintReadOnly);
+	static int SetPropertyValueByName(FPyWrapperObject* InSelf, PyObject* InValue, const FName InPropName, const char* InPythonAttrName, const bool InNotifyChange = false, const uint64 InReadOnlyFlags = CPF_EditConst | CPF_BlueprintReadOnly);
+
+	/** Set a property value on this instance (called via generated code) */
+	static int SetPropertyValue(FPyWrapperObject* InSelf, PyObject* InValue, const UProperty* InProp, const char* InPythonAttrName, const bool InNotifyChange = false, const uint64 InReadOnlyFlags = CPF_EditConst | CPF_BlueprintReadOnly);
 
 	/** Call a named getter function on this class using the given instance (called via generated code) */
-	static PyObject* CallGetterFunction(UClass* InClass, FPyWrapperObject* InSelf, const FName InFuncName);
+	static PyObject* CallGetterFunction(FPyWrapperObject* InSelf, const PyGenUtil::FGeneratedWrappedFunction& InFuncDef);
 
 	/** Call a named setter function on this class using the given instance (called via generated code) */
-	static int CallSetterFunction(UClass* InClass, FPyWrapperObject* InSelf, PyObject* InValue, const FName InFuncName);
-
-	/** Call a named function on this class (called via generated code) */
-	static PyObject* CallFunction(UClass* InClass, PyTypeObject* InType, const FName InFuncName);
+	static int CallSetterFunction(FPyWrapperObject* InSelf, PyObject* InValue, const PyGenUtil::FGeneratedWrappedFunction& InFuncDef);
 
-	/** Call a named function on this class (called via generated code) */
-	static PyObject* CallFunction(UClass* InClass, PyTypeObject* InType, PyObject* InArgs, PyObject* InKwds, const FName InFuncName, const char* InPythonFuncName, const TArray<PyGenUtil::FGeneratedWrappedMethodParameter>& InParamDef);
+	/** Call a function on this class (called via generated code) */
+	static PyObject* CallFunction(PyTypeObject* InType, const PyGenUtil::FGeneratedWrappedFunction& InFuncDef);
 
-	/** Call a named function on this class using the given instance (called via generated code) */
-	static PyObject* CallFunction(UClass* InClass, FPyWrapperObject* InSelf, const FName InFuncName);
+	/** Call a function on this class (called via generated code) */
+	static PyObject* CallFunction(PyTypeObject* InType, PyObject* InArgs, PyObject* InKwds, const PyGenUtil::FGeneratedWrappedFunction& InFuncDef, const char* InPythonFuncName);
 
-	/** Call a named function on this class using the given instance (called via generated code) */
-	static PyObject* CallFunction(UClass* InClass, FPyWrapperObject* InSelf, PyObject* InArgs, PyObject* InKwds, const FName InFuncName, const char* InPythonFuncName, const TArray<PyGenUtil::FGeneratedWrappedMethodParameter>& InParamDef);
+	/** Call a function on this class using the given instance (called via generated code) */
+	static PyObject* CallFunction(FPyWrapperObject* InSelf, const PyGenUtil::FGeneratedWrappedFunction& InFuncDef);
 
-	/** Call a named function on this instance (CallFunction internal use only) */
-	static PyObject* CallFunction_Impl(UClass* InClass, UObject* InObj, const FName InFuncName, const TCHAR* InErrorCtxt);
+	/** Call a function on this class using the given instance (called via generated code) */
+	static PyObject* CallFunction(FPyWrapperObject* InSelf, PyObject* InArgs, PyObject* InKwds, const PyGenUtil::FGeneratedWrappedFunction& InFuncDef, const char* InPythonFuncName);
 
-	/** Call a named function on this instance (CallFunction internal use only) */
-	static PyObject* CallFunction_Impl(UClass* InClass, UObject* InObj, PyObject* InArgs, PyObject* InKwds, const FName InFuncName, const char* InPythonFuncName, const TArray<PyGenUtil::FGeneratedWrappedMethodParameter>& InParamDef, const TCHAR* InErrorCtxt);
+	/** Call a function on this instance (CallFunction internal use only) */
+	static PyObject* CallFunction_Impl(UObject* InObj, const PyGenUtil::FGeneratedWrappedFunction& InFuncDef, const TCHAR* InErrorCtxt);
 
-	/** Common handler code for applying default values to function arguments (CallFunction internal use only) */
-	static void CallFunction_ApplyDefaults(UFunction* InFunc, void* InBaseParamsAddr, const TArray<PyGenUtil::FGeneratedWrappedMethodParameter>& InParamDef);
-
-	/** Common handler code for invoking a function call (CallFunction internal use only) */
-	static void CallFunction_InvokeImpl(UObject* InObj, UFunction* InFunc, void* InBaseParamsAddr, const TCHAR* InErrorCtxt);
-
-	/** Common handler code for processing return values from a function call (CallFunction internal use only) */
-	static PyObject* CallFunction_ReturnImpl(UObject* InObj, UFunction* InFunc, const void* InBaseParamsAddr, const TCHAR* InErrorCtxt);
+	/** Call a function on this instance (CallFunction internal use only) */
+	static PyObject* CallFunction_Impl(UObject* InObj, PyObject* InArgs, PyObject* InKwds, const PyGenUtil::FGeneratedWrappedFunction& InFuncDef, const char* InPythonFuncName, const TCHAR* InErrorCtxt);
 
 	/** Implementation of the "call" logic for a Python class method with no arguments (internal Python bindings use only) */
 	static PyObject* CallClassMethodNoArgs_Impl(PyTypeObject* InType, void* InClosure);
@@ -104,9 +101,12 @@ struct FPyWrapperObject : public FPyWrapperBase
 /** Meta-data for all UE4 exposed object types */
 struct FPyWrapperObjectMetaData : public FPyWrapperBaseMetaData
 {
-	PY_OVERRIDE_GETSET_METADATA(FPyWrapperObjectMetaData)
+	PY_METADATA_METHODS(FPyWrapperObjectMetaData, FGuid(0x89FC2465, 0xA83F4F31, 0xBBCC1E86, 0xE9D76551))
 
-	FPyWrapperObjectMetaData();
+	FPyWrapperObjectMetaData()
+		: Class(nullptr)
+	{
+	}
 
 	/** Get the UClass from the given type */
 	static UClass* GetClass(PyTypeObject* PyType);
@@ -126,6 +126,9 @@ struct FPyWrapperObjectMetaData : public FPyWrapperBaseMetaData
 	/** Resolve the original function name of a Python method of the given instance */
 	static FName ResolveFunctionName(FPyWrapperObject* Instance, const FName InPythonMethodName);
 
+	/** Add object references from the given Python object to the given collector */
+	virtual void AddReferencedObjects(FPyWrapperBase* Instance, FReferenceCollector& Collector) override;
+
 	/** Unreal class */
 	UClass* Class;
 
diff --git a/Engine/Plugins/Experimental/PythonScriptPlugin/Source/PythonScriptPlugin/Private/PyWrapperOwnerContext.cpp b/Engine/Plugins/Experimental/PythonScriptPlugin/Source/PythonScriptPlugin/Private/PyWrapperOwnerContext.cpp
--- a/Engine/Plugins/Experimental/PythonScriptPlugin/Source/PythonScriptPlugin/Private/PyWrapperOwnerContext.cpp
+++ b/Engine/Plugins/Experimental/PythonScriptPlugin/Source/PythonScriptPlugin/Private/PyWrapperOwnerContext.cpp
@@ -5,27 +5,29 @@
 #if WITH_PYTHON
 
 FPyWrapperOwnerContext::FPyWrapperOwnerContext()
+	: OwnerObject()
+	, OwnerProperty(nullptr)
 {
 }
 
-FPyWrapperOwnerContext::FPyWrapperOwnerContext(PyObject* InOwner, const FName InPropName)
+FPyWrapperOwnerContext::FPyWrapperOwnerContext(PyObject* InOwner, const UProperty* InProp)
 	: OwnerObject(FPyObjectPtr::NewReference(InOwner))
-	, OwnerPropertyName(InPropName)
+	, OwnerProperty(InProp)
 {
-	checkf(OwnerPropertyName.IsNone() || OwnerObject.IsValid(), TEXT("Owner context cannot have an owner property without an owner object"));
+	checkf(!OwnerProperty || OwnerObject.IsValid(), TEXT("Owner context cannot have an owner property without an owner object"));
 }
 
-FPyWrapperOwnerContext::FPyWrapperOwnerContext(const FPyObjectPtr& InOwner, const FName InPropName)
+FPyWrapperOwnerContext::FPyWrapperOwnerContext(const FPyObjectPtr& InOwner, const UProperty* InProp)
 	: OwnerObject(InOwner)
-	, OwnerPropertyName(InPropName)
+	, OwnerProperty(InProp)
 {
-	checkf(OwnerPropertyName.IsNone() || OwnerObject.IsValid(), TEXT("Owner context cannot have an owner property without an owner object"));
+	checkf(!OwnerProperty || OwnerObject.IsValid(), TEXT("Owner context cannot have an owner property without an owner object"));
 }
 
 void FPyWrapperOwnerContext::Reset()
 {
 	OwnerObject.Reset();
-	OwnerPropertyName = NAME_None;
+	OwnerProperty = nullptr;
 }
 
 bool FPyWrapperOwnerContext::HasOwner() const
@@ -38,9 +40,9 @@ PyObject* FPyWrapperOwnerContext::GetOwnerObject() const
 	return (PyObject*)OwnerObject.GetPtr();
 }
 
-FName FPyWrapperOwnerContext::GetOwnerPropertyName() const
+const UProperty* FPyWrapperOwnerContext::GetOwnerProperty() const
 {
-	return OwnerPropertyName;
+	return OwnerProperty;
 }
 
 void FPyWrapperOwnerContext::AssertValidConversionMethod(const EPyConversionMethod InMethod) const
diff --git a/Engine/Plugins/Experimental/PythonScriptPlugin/Source/PythonScriptPlugin/Private/PyWrapperOwnerContext.h b/Engine/Plugins/Experimental/PythonScriptPlugin/Source/PythonScriptPlugin/Private/PyWrapperOwnerContext.h
--- a/Engine/Plugins/Experimental/PythonScriptPlugin/Source/PythonScriptPlugin/Private/PyWrapperOwnerContext.h
+++ b/Engine/Plugins/Experimental/PythonScriptPlugin/Source/PythonScriptPlugin/Private/PyWrapperOwnerContext.h
@@ -15,11 +15,11 @@ public:
 	/** Default constructor */
 	FPyWrapperOwnerContext();
 
-	/** Construct this context from the given Python object and optional property name (will create a new reference to the given object) */
-	explicit FPyWrapperOwnerContext(PyObject* InOwner, const FName InPropName = NAME_None);
+	/** Construct this context from the given Python object and optional property (will create a new reference to the given object) */
+	explicit FPyWrapperOwnerContext(PyObject* InOwner, const UProperty* InProp = nullptr);
 
-	/** Construct this context from the given Python object and optional property name */
-	explicit FPyWrapperOwnerContext(const FPyObjectPtr& InOwner, const FName InPropName = NAME_None);
+	/** Construct this context from the given Python object and optional property */
+	explicit FPyWrapperOwnerContext(const FPyObjectPtr& InOwner, const UProperty* InProp = nullptr);
 
 	/** Reset this context back to its default state */
 	void Reset();
@@ -30,8 +30,8 @@ public:
 	/** Get the Python object that owns the instance being wrapped (if any, borrowed reference) */
 	PyObject* GetOwnerObject() const;
 
-	/** Get the name of the property on the owner object that owns the instance being wrapped (if known) */
-	FName GetOwnerPropertyName() const;
+	/** Get the property on the owner object that owns the instance being wrapped (if known) */
+	const UProperty* GetOwnerProperty() const;
 
 	/** Assert that the given conversion method is valid for this owner context */
 	void AssertValidConversionMethod(const EPyConversionMethod InMethod) const;
@@ -40,8 +40,8 @@ private:
 	/** The Python object that owns the instance being wrapped (if any) */
 	FPyObjectPtr OwnerObject;
 
-	/** The name of the property on the owner object that owns the instance being wrapped (if known) */
-	FName OwnerPropertyName;
+	/** The property on the owner object that owns the instance being wrapped (if known) */
+	const UProperty* OwnerProperty;
 };
 
 #endif	// WITH_PYTHON
diff --git a/Engine/Plugins/Experimental/PythonScriptPlugin/Source/PythonScriptPlugin/Private/PyWrapperSet.cpp b/Engine/Plugins/Experimental/PythonScriptPlugin/Source/PythonScriptPlugin/Private/PyWrapperSet.cpp
--- a/Engine/Plugins/Experimental/PythonScriptPlugin/Source/PythonScriptPlugin/Private/PyWrapperSet.cpp
+++ b/Engine/Plugins/Experimental/PythonScriptPlugin/Source/PythonScriptPlugin/Private/PyWrapperSet.cpp
@@ -327,10 +327,14 @@ bool FPyWrapperSet::ValidateInternalState(FPyWrapperSet* InSelf)
 	return true;
 }
 
-FPyWrapperSet* FPyWrapperSet::CastPyObject(PyObject* InPyObject)
+FPyWrapperSet* FPyWrapperSet::CastPyObject(PyObject* InPyObject, FPyConversionResult* OutCastResult)
 {
+	SetOptionalPyConversionResult(FPyConversionResult::Failure(), OutCastResult);
+
 	if (PyObject_IsInstance(InPyObject, (PyObject*)&PyWrapperSetType) == 1)
 	{
+		SetOptionalPyConversionResult(FPyConversionResult::Success(), OutCastResult);
+
 		Py_INCREF(InPyObject);
 		return (FPyWrapperSet*)InPyObject;
 	}
@@ -338,8 +342,10 @@ FPyWrapperSet* FPyWrapperSet::CastPyObject(PyObject* InPyObject)
 	return nullptr;
 }
 
-FPyWrapperSet* FPyWrapperSet::CastPyObject(PyObject* InPyObject, PyTypeObject* InType, const PyUtil::FPropertyDef& InElementDef)
+FPyWrapperSet* FPyWrapperSet::CastPyObject(PyObject* InPyObject, PyTypeObject* InType, const PyUtil::FPropertyDef& InElementDef, FPyConversionResult* OutCastResult)
 {
+	SetOptionalPyConversionResult(FPyConversionResult::Failure(), OutCastResult);
+
 	if (PyObject_IsInstance(InPyObject, (PyObject*)InType) == 1 && (InType == &PyWrapperSetType || PyObject_IsInstance(InPyObject, (PyObject*)&PyWrapperSetType) == 1))
 	{
 		FPyWrapperSet* Self = (FPyWrapperSet*)InPyObject;
@@ -352,6 +358,8 @@ FPyWrapperSet* FPyWrapperSet::CastPyObject(PyObject* InPyObject, PyTypeObject* I
 		const PyUtil::FPropertyDef SelfElementPropDef = Self->SetProp->ElementProp;
 		if (SelfElementPropDef == InElementDef)
 		{
+			SetOptionalPyConversionResult(FPyConversionResult::Success(), OutCastResult);
+
 			Py_INCREF(Self);
 			return Self;
 		}
@@ -397,6 +405,7 @@ FPyWrapperSet* FPyWrapperSet::CastPyObject(PyObject* InPyObject, PyTypeObject* I
 			NewScriptSetHelper.Rehash();
 		}
 
+		SetOptionalPyConversionResult(FPyConversionResult::SuccessWithCoercion(), OutCastResult);
 		return NewSet.Release();
 	}
 
@@ -432,6 +441,8 @@ FPyWrapperSet* FPyWrapperSet::CastPyObject(PyObject* InPyObject, PyTypeObject* I
 				}
 
 				NewScriptSetHelper.Rehash();
+
+				SetOptionalPyConversionResult(FPyConversionResult::SuccessWithCoercion(), OutCastResult);
 				return NewSet.Release();
 			}
 		}
@@ -1586,17 +1597,14 @@ PyTypeObject InitializePyWrapperSetIteratorType()
 PyTypeObject PyWrapperSetType = InitializePyWrapperSetType();
 PyTypeObject PyWrapperSetIteratorType = InitializePyWrapperSetIteratorType();
 
-FPyWrapperSetMetaData::FPyWrapperSetMetaData()
+void FPyWrapperSetMetaData::AddReferencedObjects(FPyWrapperBase* Instance, FReferenceCollector& Collector)
 {
-	AddReferencedObjects = [](FPyWrapperBase* Self, FReferenceCollector& Collector)
+	FPyWrapperSet* Self = static_cast<FPyWrapperSet*>(Instance);
+	if (Self->SetProp && Self->SetInstance && !Self->OwnerContext.HasOwner())
 	{
-		FPyWrapperSet* SetSelf = static_cast<FPyWrapperSet*>(Self);
-		if (SetSelf->SetProp && SetSelf->SetInstance && !SetSelf->OwnerContext.HasOwner())
-		{
-			Collector.AddReferencedObject(SetSelf->SetProp);
-			FPyReferenceCollector::AddReferencedObjectsFromProperty(Collector, SetSelf->SetProp, SetSelf->SetInstance);
-		}
-	};
+		Collector.AddReferencedObject(Self->SetProp);
+		FPyReferenceCollector::AddReferencedObjectsFromProperty(Collector, Self->SetProp, Self->SetInstance);
+	}
 }
 
 #endif	// WITH_PYTHON
diff --git a/Engine/Plugins/Experimental/PythonScriptPlugin/Source/PythonScriptPlugin/Private/PyWrapperSet.h b/Engine/Plugins/Experimental/PythonScriptPlugin/Source/PythonScriptPlugin/Private/PyWrapperSet.h
--- a/Engine/Plugins/Experimental/PythonScriptPlugin/Source/PythonScriptPlugin/Private/PyWrapperSet.h
+++ b/Engine/Plugins/Experimental/PythonScriptPlugin/Source/PythonScriptPlugin/Private/PyWrapperSet.h
@@ -45,10 +45,10 @@ struct FPyWrapperSet : public FPyWrapperBase
 	static bool ValidateInternalState(FPyWrapperSet* InSelf);
 
 	/** Cast the given Python object to this wrapped type (returns a new reference) */
-	static FPyWrapperSet* CastPyObject(PyObject* InPyObject);
+	static FPyWrapperSet* CastPyObject(PyObject* InPyObject, FPyConversionResult* OutCastResult = nullptr);
 
 	/** Cast the given Python object to this wrapped type, or attempt to convert the type into a new wrapped instance (returns a new reference) */
-	static FPyWrapperSet* CastPyObject(PyObject* InPyObject, PyTypeObject* InType, const PyUtil::FPropertyDef& InElementDef);
+	static FPyWrapperSet* CastPyObject(PyObject* InPyObject, PyTypeObject* InType, const PyUtil::FPropertyDef& InElementDef, FPyConversionResult* OutCastResult = nullptr);
 
 	/** Get the length of this container (equivalent to 'len(x)' in Python) */
 	static Py_ssize_t Len(FPyWrapperSet* InSelf);
@@ -108,9 +108,14 @@ struct FPyWrapperSet : public FPyWrapperBase
 /** Meta-data for all UE4 exposed set types */
 struct FPyWrapperSetMetaData : public FPyWrapperBaseMetaData
 {
-	PY_OVERRIDE_GETSET_METADATA(FPyWrapperSetMetaData)
+	PY_METADATA_METHODS(FPyWrapperSetMetaData, FGuid(0xB31DDC32, 0xAEA548BE, 0xB6FB5AB2, 0x02E5CA83))
 
-	FPyWrapperSetMetaData();
+	FPyWrapperSetMetaData()
+	{
+	}
+
+	/** Add object references from the given Python object to the given collector */
+	virtual void AddReferencedObjects(FPyWrapperBase* Instance, FReferenceCollector& Collector) override;
 };
 
 typedef TPyPtr<FPyWrapperSet> FPyWrapperSetPtr;
