Copying //UE4/Dev-Enterprise to //UE4/Dev-Main (Source: //UE4/Dev-Enterprise @ 3972172)

diff --git a/Engine/Plugins/Experimental/PythonScriptPlugin/Source/PythonScriptPlugin/Private/PyWrapperStruct.cpp b/Engine/Plugins/Experimental/PythonScriptPlugin/Source/PythonScriptPlugin/Private/PyWrapperStruct.cpp
--- a/Engine/Plugins/Experimental/PythonScriptPlugin/Source/PythonScriptPlugin/Private/PyWrapperStruct.cpp
+++ b/Engine/Plugins/Experimental/PythonScriptPlugin/Source/PythonScriptPlugin/Private/PyWrapperStruct.cpp
@@ -177,10 +177,14 @@ bool FPyWrapperStruct::ValidateInternalState(FPyWrapperStruct* InSelf)
 	return true;
 }
 
-FPyWrapperStruct* FPyWrapperStruct::CastPyObject(PyObject* InPyObject)
+FPyWrapperStruct* FPyWrapperStruct::CastPyObject(PyObject* InPyObject, FPyConversionResult* OutCastResult)
 {
+	SetOptionalPyConversionResult(FPyConversionResult::Failure(), OutCastResult);
+
 	if (PyObject_IsInstance(InPyObject, (PyObject*)&PyWrapperStructType) == 1)
 	{
+		SetOptionalPyConversionResult(FPyConversionResult::Success(), OutCastResult);
+
 		Py_INCREF(InPyObject);
 		return (FPyWrapperStruct*)InPyObject;
 	}
@@ -188,10 +192,14 @@ FPyWrapperStruct* FPyWrapperStruct::CastPyObject(PyObject* InPyObject)
 	return nullptr;
 }
 
-FPyWrapperStruct* FPyWrapperStruct::CastPyObject(PyObject* InPyObject, PyTypeObject* InType)
+FPyWrapperStruct* FPyWrapperStruct::CastPyObject(PyObject* InPyObject, PyTypeObject* InType, FPyConversionResult* OutCastResult)
 {
+	SetOptionalPyConversionResult(FPyConversionResult::Failure(), OutCastResult);
+
 	if (PyObject_IsInstance(InPyObject, (PyObject*)InType) == 1 && (InType == &PyWrapperStructType || PyObject_IsInstance(InPyObject, (PyObject*)&PyWrapperStructType) == 1))
 	{
+		SetOptionalPyConversionResult(Py_TYPE(InPyObject) == InType ? FPyConversionResult::Success() : FPyConversionResult::SuccessWithCoercion(), OutCastResult);
+
 		Py_INCREF(InPyObject);
 		return (FPyWrapperStruct*)InPyObject;
 	}
@@ -226,7 +234,7 @@ FPyWrapperStruct* FPyWrapperStruct::CastPyObject(PyObject* InPyObject, PyTypeObj
 					break;
 				}
 
-				const int Result = PyUtil::SetUEPropValue(NewStruct->ScriptStruct, NewStruct->StructInstance, SequenceItem, InitParam.ParamPropName, InitParam.ParamName.GetData(), FPyWrapperOwnerContext(), 0, *PyUtil::GetErrorContext(NewStruct.Get()));
+				const int Result = PyUtil::SetUEPropValue(NewStruct->ScriptStruct, NewStruct->StructInstance, SequenceItem, InitParam.ParamProp, InitParam.ParamName.GetData(), FPyWrapperOwnerContext(), 0, *PyUtil::GetErrorContext(NewStruct.Get()));
 				if (Result != 0)
 				{
 					return nullptr;
@@ -234,6 +242,7 @@ FPyWrapperStruct* FPyWrapperStruct::CastPyObject(PyObject* InPyObject, PyTypeObj
 			}
 		}
 
+		SetOptionalPyConversionResult(FPyConversionResult::SuccessWithCoercion(), OutCastResult);
 		return NewStruct.Release();
 	}
 
@@ -257,7 +266,7 @@ FPyWrapperStruct* FPyWrapperStruct::CastPyObject(PyObject* InPyObject, PyTypeObj
 			PyObject* MappingItem = PyMapping_GetItemString(InPyObject, (char*)InitParam.ParamName.GetData());
 			if (MappingItem)
 			{
-				const int Result = PyUtil::SetUEPropValue(NewStruct->ScriptStruct, NewStruct->StructInstance, MappingItem, InitParam.ParamPropName, InitParam.ParamName.GetData(), FPyWrapperOwnerContext(), 0, *PyUtil::GetErrorContext(NewStruct.Get()));
+				const int Result = PyUtil::SetUEPropValue(NewStruct->ScriptStruct, NewStruct->StructInstance, MappingItem, InitParam.ParamProp, InitParam.ParamName.GetData(), FPyWrapperOwnerContext(), 0, *PyUtil::GetErrorContext(NewStruct.Get()));
 				if (Result != 0)
 				{
 					return nullptr;
@@ -265,6 +274,7 @@ FPyWrapperStruct* FPyWrapperStruct::CastPyObject(PyObject* InPyObject, PyTypeObj
 			}
 		}
 
+		SetOptionalPyConversionResult(FPyConversionResult::SuccessWithCoercion(), OutCastResult);
 		return NewStruct.Release();
 	}
 
@@ -296,7 +306,7 @@ int FPyWrapperStruct::SetPropertyValues(FPyWrapperStruct* InSelf, PyObject* InAr
 		if (PyValue)
 		{
 			const PyGenUtil::FGeneratedWrappedMethodParameter& InitParam = StructMetaData->InitParams[ParamIndex];
-			const int Result = PyUtil::SetUEPropValue(InSelf->ScriptStruct, InSelf->StructInstance, PyValue, InitParam.ParamPropName, InitParam.ParamName.GetData(), FPyWrapperOwnerContext(), 0, *PyUtil::GetErrorContext(InSelf));
+			const int Result = PyUtil::SetUEPropValue(InSelf->ScriptStruct, InSelf->StructInstance, PyValue, InitParam.ParamProp, InitParam.ParamName.GetData(), FPyWrapperOwnerContext(), 0, *PyUtil::GetErrorContext(InSelf));
 			if (Result != 0)
 			{
 				return Result;
@@ -307,37 +317,262 @@ int FPyWrapperStruct::SetPropertyValues(FPyWrapperStruct* InSelf, PyObject* InAr
 	return 0;
 }
 
-PyObject* FPyWrapperStruct::GetPropertyValue(FPyWrapperStruct* InSelf, const FName InPropName, const char* InPythonAttrName)
+PyObject* FPyWrapperStruct::GetPropertyValueByName(FPyWrapperStruct* InSelf, const FName InPropName, const char* InPythonAttrName)
+{
+	if (!ValidateInternalState(InSelf))
+	{
+		return nullptr;
+	}
+
+	UProperty* Prop = InSelf->ScriptStruct->FindPropertyByName(InPropName);
+	if (!Prop)
+	{
+		PyUtil::SetPythonError(PyExc_Exception, InSelf, *FString::Printf(TEXT("Failed to find property '%s' for attribute '%s' on '%s'"), *InPropName.ToString(), UTF8_TO_TCHAR(InPythonAttrName), *InSelf->ScriptStruct->GetName()));
+		return nullptr;
+	}
+
+	return GetPropertyValue(InSelf, Prop, InPythonAttrName);
+}
+
+PyObject* FPyWrapperStruct::GetPropertyValue(FPyWrapperStruct* InSelf, const UProperty* InProp, const char* InPythonAttrName)
 {
 	if (!ValidateInternalState(InSelf))
 	{
 		return nullptr;
 	}
 
-	return PyUtil::GetUEPropValue(InSelf->ScriptStruct, InSelf->StructInstance, InPropName, InPythonAttrName, (PyObject*)InSelf, *PyUtil::GetErrorContext(InSelf));
+	return PyUtil::GetUEPropValue(InSelf->ScriptStruct, InSelf->StructInstance, InProp, InPythonAttrName, (PyObject*)InSelf, *PyUtil::GetErrorContext(InSelf));
+}
+
+int FPyWrapperStruct::SetPropertyValueByName(FPyWrapperStruct* InSelf, PyObject* InValue, const FName InPropName, const char* InPythonAttrName, const bool InNotifyChange, const uint64 InReadOnlyFlags)
+{
+	if (!ValidateInternalState(InSelf))
+	{
+		return -1;
+	}
+
+	UProperty* Prop = InSelf->ScriptStruct->FindPropertyByName(InPropName);
+	if (!Prop)
+	{
+		PyUtil::SetPythonError(PyExc_Exception, InSelf, *FString::Printf(TEXT("Failed to find property '%s' for attribute '%s' on '%s'"), *InPropName.ToString(), UTF8_TO_TCHAR(InPythonAttrName), *InSelf->ScriptStruct->GetName()));
+		return -1;
+	}
+
+	return SetPropertyValue(InSelf, InValue, Prop, InPythonAttrName, InNotifyChange, InReadOnlyFlags);
 }
 
-int FPyWrapperStruct::SetPropertyValue(FPyWrapperStruct* InSelf, PyObject* InValue, const FName InPropName, const char* InPythonAttrName, const bool InNotifyChange, const uint64 InReadOnlyFlags)
+int FPyWrapperStruct::SetPropertyValue(FPyWrapperStruct* InSelf, PyObject* InValue, const UProperty* InProp, const char* InPythonAttrName, const bool InNotifyChange, const uint64 InReadOnlyFlags)
 {
 	if (!ValidateInternalState(InSelf))
 	{
 		return -1;
 	}
 
-	const FPyWrapperOwnerContext ChangeOwner = InNotifyChange ? FPyWrapperOwnerContext((PyObject*)InSelf, InPropName) : FPyWrapperOwnerContext();
-	return PyUtil::SetUEPropValue(InSelf->ScriptStruct, InSelf->StructInstance, InValue, InPropName, InPythonAttrName, ChangeOwner, InReadOnlyFlags, *PyUtil::GetErrorContext(InSelf));
+	const FPyWrapperOwnerContext ChangeOwner = InNotifyChange ? FPyWrapperOwnerContext((PyObject*)InSelf, InProp) : FPyWrapperOwnerContext();
+	return PyUtil::SetUEPropValue(InSelf->ScriptStruct, InSelf->StructInstance, InValue, InProp, InPythonAttrName, ChangeOwner, InReadOnlyFlags, *PyUtil::GetErrorContext(InSelf));
+}
+
+PyObject* FPyWrapperStruct::CallFunction_Impl(FPyWrapperStruct* InSelf, PyObject* InArgs, PyObject* InKwds, const PyGenUtil::FGeneratedWrappedFunction& InFuncDef, const PyGenUtil::FGeneratedWrappedMethodParameter& InStructParam, const char* InPythonFuncName, const TCHAR* InErrorCtxt)
+{
+	TArray<PyObject*> Params;
+	if ((InArgs || InKwds) && !PyGenUtil::ParseMethodParameters(InArgs, InKwds, InFuncDef.InputParams, InPythonFuncName, Params))
+	{
+		return nullptr;
+	}
+
+	if (ensureAlways(InFuncDef.Func))
+	{
+		UClass* Class = InFuncDef.Func->GetOwnerClass();
+		UObject* Obj = Class->GetDefaultObject();
+
+		// Deprecated functions emit a warning
+		{
+			FString DeprecationMessage;
+			if (PyGenUtil::IsDeprecatedFunction(InFuncDef.Func, &DeprecationMessage) && 
+				PyUtil::SetPythonWarning(PyExc_DeprecationWarning, InErrorCtxt, *FString::Printf(TEXT("Function '%s.%s' is deprecated: %s"), *Class->GetName(), *InFuncDef.Func->GetName(), *DeprecationMessage)) == -1
+				)
+			{
+				// -1 from SetPythonWarning means the warning should be an exception
+				return nullptr;
+			}
+		}
+
+		FStructOnScope FuncParams(InFuncDef.Func);
+		PyGenUtil::ApplyParamDefaults(FuncParams.GetStructMemory(), InFuncDef.InputParams);
+		if (ensureAlways(InStructParam.ParamProp))
+		{
+			void* StructArgInstance = InStructParam.ParamProp->ContainerPtrToValuePtr<void>(FuncParams.GetStructMemory());
+			CastChecked<UStructProperty>(InStructParam.ParamProp)->Struct->CopyScriptStruct(StructArgInstance, InSelf->StructInstance);
+		}
+		for (int32 ParamIndex = 0; ParamIndex < Params.Num(); ++ParamIndex)
+		{
+			const PyGenUtil::FGeneratedWrappedMethodParameter& ParamDef = InFuncDef.InputParams[ParamIndex];
+
+			PyObject* PyValue = Params[ParamIndex];
+			if (PyValue)
+			{
+				if (!PyConversion::NativizeProperty_InContainer(PyValue, ParamDef.ParamProp, FuncParams.GetStructMemory(), 0))
+				{
+					PyUtil::SetPythonError(PyExc_TypeError, InErrorCtxt, *FString::Printf(TEXT("Failed to convert parameter '%s.%s' when calling function '%s' on '%s'"), UTF8_TO_TCHAR(ParamDef.ParamName.GetData()), *Class->GetName(), *InFuncDef.Func->GetName(), *Obj->GetName()));
+					return nullptr;
+				}
+			}
+		}
+		PyUtil::InvokeFunctionCall(Obj, InFuncDef.Func, FuncParams.GetStructMemory(), InErrorCtxt);
+		return PyGenUtil::PackReturnValues(FuncParams.GetStructMemory(), InFuncDef.OutputParams, InErrorCtxt, *FString::Printf(TEXT("function '%s.%s' on '%s'"), *Class->GetName(), *InFuncDef.Func->GetName(), *Obj->GetName()));
+	}
+
+	Py_RETURN_NONE;
+}
+
+PyObject* FPyWrapperStruct::CallMethodNoArgs_Impl(FPyWrapperStruct* InSelf, void* InClosure)
+{
+	if (!ValidateInternalState(InSelf))
+	{
+		return nullptr;
+	}
+
+	const PyGenUtil::FGeneratedWrappedDynamicStructMethod* Closure = (PyGenUtil::FGeneratedWrappedDynamicStructMethod*)InClosure;
+	return CallFunction_Impl(InSelf, nullptr, nullptr, Closure->MethodFunc, Closure->StructParam, Closure->MethodName.GetData(), *PyUtil::GetErrorContext(InSelf));
+}
+
+PyObject* FPyWrapperStruct::CallMethodWithArgs_Impl(FPyWrapperStruct* InSelf, PyObject* InArgs, PyObject* InKwds, void* InClosure)
+{
+	if (!ValidateInternalState(InSelf))
+	{
+		return nullptr;
+	}
+
+	const PyGenUtil::FGeneratedWrappedDynamicStructMethod* Closure = (PyGenUtil::FGeneratedWrappedDynamicStructMethod*)InClosure;
+	return CallFunction_Impl(InSelf, InArgs, InKwds, Closure->MethodFunc, Closure->StructParam, Closure->MethodName.GetData(), *PyUtil::GetErrorContext(InSelf));
+}
+
+PyObject* FPyWrapperStruct::CallBinaryOperatorFunction_Impl(FPyWrapperStruct* InSelf, PyObject* InRHS, const PyGenUtil::FGeneratedWrappedStructMathOpFunction& InMathOpFunc, const bool InInlineOp, const TOptional<EPyConversionResultState> InRequiredConversionResult, FPyConversionResult* OutRHSConversionResult)
+{
+	SetOptionalPyConversionResult(FPyConversionResult::Failure(), OutRHSConversionResult);
+
+	// Binary functions must have a single input parameter (excluding the struct parameter) and a single output parameter (the return value)
+	if (InMathOpFunc.InputParams.Num() != 1 || InMathOpFunc.OutputParams.Num() != 1)
+	{
+		return nullptr;
+	}
+
+	if (ensureAlways(InMathOpFunc.Func))
+	{
+		UClass* Class = InMathOpFunc.Func->GetOwnerClass();
+		UObject* Obj = Class->GetDefaultObject();
+
+		// Build the input arguments (failures here aren't fatal as we may have multiple functions to evaluate on the stack, only one of which may accept the RHS parameter)
+		FStructOnScope FuncParams(InMathOpFunc.Func);
+		{
+			const FPyConversionResult RHSResult = PyConversion::NativizeProperty_InContainer(InRHS, InMathOpFunc.InputParams[0].ParamProp, FuncParams.GetStructMemory(), 0, FPyWrapperOwnerContext(), PyConversion::ESetErrorState::No);
+			SetOptionalPyConversionResult(RHSResult, OutRHSConversionResult);
+
+			if (!RHSResult)
+			{
+				return nullptr;
+			}
+
+			if (InRequiredConversionResult.IsSet() && RHSResult.GetState() != InRequiredConversionResult.GetValue())
+			{
+				return nullptr;
+			}
+		}
+		if (ensureAlways(InMathOpFunc.StructParam.ParamProp))
+		{
+			void* StructArgInstance = InMathOpFunc.StructParam.ParamProp->ContainerPtrToValuePtr<void>(FuncParams.GetStructMemory());
+			CastChecked<UStructProperty>(InMathOpFunc.StructParam.ParamProp)->Struct->CopyScriptStruct(StructArgInstance, InSelf->StructInstance);
+		}
+		PyUtil::InvokeFunctionCall(Obj, InMathOpFunc.Func, FuncParams.GetStructMemory(), *PyUtil::GetErrorContext(InSelf));
+
+		PyObject* ReturnPyObj = nullptr;
+		const PyGenUtil::FGeneratedWrappedMethodParameter& ReturnParamDef = InMathOpFunc.OutputParams[0];
+		if (InInlineOp)
+		{
+			// Copy the result back into ourself
+			if (const UStructProperty* ReturnStructProp = Cast<const UStructProperty>(ReturnParamDef.ParamProp))
+			{
+				if (ReturnStructProp->Struct == FPyWrapperStructMetaData::GetStruct(InSelf))
+				{
+					void* ReturnStructInstance = ReturnStructProp->ContainerPtrToValuePtr<void>(FuncParams.GetStructMemory());
+					ReturnStructProp->Struct->CopyScriptStruct(InSelf->StructInstance, ReturnStructInstance);
+
+					Py_INCREF(InSelf);
+					ReturnPyObj = (PyObject*)InSelf;
+				}
+			}
+		}
+		else
+		{
+			PyConversion::PythonizeProperty_InContainer(ReturnParamDef.ParamProp, FuncParams.GetStructMemory(), 0, ReturnPyObj, EPyConversionMethod::Steal);
+		}
+
+		if (!ReturnPyObj)
+		{
+			PyUtil::SetPythonError(PyExc_TypeError, InSelf, *FString::Printf(TEXT("Failed to convert return property '%s' (%s) when calling function '%s' on '%s'"), *ReturnParamDef.ParamProp->GetName(), *ReturnParamDef.ParamProp->GetClass()->GetName(), *InMathOpFunc.Func->GetName(), *Obj->GetName()));
+		}
+
+		return ReturnPyObj;
+	}
+
+	return nullptr;
+}
+
+PyObject* FPyWrapperStruct::CallBinaryOperator_Impl(FPyWrapperStruct* InSelf, PyObject* InRHS, const PyGenUtil::FGeneratedWrappedStructMathOpStack::EOpType InOpType)
+{
+	if (!ValidateInternalState(InSelf))
+	{
+		return nullptr;
+	}
+
+	FPyWrapperStructMetaData* StructMetaData = FPyWrapperStructMetaData::GetMetaData(InSelf);
+	if (!StructMetaData)
+	{
+		return nullptr;
+	}
+
+	// We process the operator stack in two passes:
+	//	- The first pass looks for a signature that exactly matches the given argument
+	//	- The second pass allows type coercion to occur when calling the signature
+	// We use the first pass to find a function that may be called for the second pass
+	const bool bInlineOp = PyGenUtil::FGeneratedWrappedStructMathOpStack::IsInlineOp(InOpType);
+	const PyGenUtil::FGeneratedWrappedStructMathOpFunction* CoercedMathOpFunc = nullptr;
+	for (const PyGenUtil::FGeneratedWrappedStructMathOpFunction& MathOpFunc : StructMetaData->MathOpStacks[(int32)InOpType].MathOpFuncs)
+	{
+		FPyConversionResult RHSConversionResult = FPyConversionResult::Failure();
+		PyObject* PyResult = CallBinaryOperatorFunction_Impl(InSelf, InRHS, MathOpFunc, bInlineOp, EPyConversionResultState::Success, &RHSConversionResult);
+		if (PyResult)
+		{
+			return PyResult;
+		}
+		else if (RHSConversionResult.GetState() == EPyConversionResultState::SuccessWithCoercion)
+		{
+			CoercedMathOpFunc = &MathOpFunc;
+		}
+	}
+	if (CoercedMathOpFunc)
+	{
+		PyObject* PyResult = CallBinaryOperatorFunction_Impl(InSelf, InRHS, *CoercedMathOpFunc, bInlineOp);
+		if (PyResult)
+		{
+			return PyResult;
+		}
+	}
+
+	Py_INCREF(Py_NotImplemented);
+	return Py_NotImplemented;
 }
 
 PyObject* FPyWrapperStruct::Getter_Impl(FPyWrapperStruct* InSelf, void* InClosure)
 {
 	const PyGenUtil::FGeneratedWrappedGetSet* Closure = (PyGenUtil::FGeneratedWrappedGetSet*)InClosure;
-	return GetPropertyValue(InSelf, Closure->PropName, Closure->GetSetName.GetData());
+	return GetPropertyValue(InSelf, Closure->Prop, Closure->GetSetName.GetData());
 }
 
 int FPyWrapperStruct::Setter_Impl(FPyWrapperStruct* InSelf, PyObject* InValue, void* InClosure)
 {
 	const PyGenUtil::FGeneratedWrappedGetSet* Closure = (PyGenUtil::FGeneratedWrappedGetSet*)InClosure;
-	return SetPropertyValue(InSelf, InValue, Closure->PropName, Closure->GetSetName.GetData());
+	return SetPropertyValue(InSelf, InValue, Closure->Prop, Closure->GetSetName.GetData());
 }
 
 PyTypeObject InitializePyWrapperStructType()
@@ -446,7 +681,7 @@ PyTypeObject InitializePyWrapperStructType()
 			}
 
 			const FName ResolvedName = FPyWrapperStructMetaData::ResolvePropertyName(InSelf, Name);
-			return FPyWrapperStruct::GetPropertyValue(InSelf, ResolvedName, TCHAR_TO_UTF8(*Name.ToString()));
+			return FPyWrapperStruct::GetPropertyValueByName(InSelf, ResolvedName, TCHAR_TO_UTF8(*Name.ToString()));
 		}
 
 		static PyObject* SetEditorProperty(FPyWrapperStruct* InSelf, PyObject* InArgs, PyObject* InKwds)
@@ -468,7 +703,7 @@ PyTypeObject InitializePyWrapperStructType()
 			}
 
 			const FName ResolvedName = FPyWrapperStructMetaData::ResolvePropertyName(InSelf, Name);
-			const int Result = FPyWrapperStruct::SetPropertyValue(InSelf, PyValueObj, ResolvedName, TCHAR_TO_UTF8(*Name.ToString()), /*InNotifyChange*/true, CPF_EditConst);
+			const int Result = FPyWrapperStruct::SetPropertyValueByName(InSelf, PyValueObj, ResolvedName, TCHAR_TO_UTF8(*Name.ToString()), /*InNotifyChange*/true, CPF_EditConst);
 			if (Result != 0)
 			{
 				return nullptr;
@@ -531,16 +766,17 @@ FPyWrapperStructMetaData::FPyWrapperStructMetaData()
 {
 	static const FPyWrapperStructAllocationPolicy_Heap HeapAllocPolicy = FPyWrapperStructAllocationPolicy_Heap();
 	AllocPolicy = &HeapAllocPolicy;
+}
 
-	AddReferencedObjects = [](FPyWrapperBase* Self, FReferenceCollector& Collector)
+/** Add object references from the given Python object to the given collector */
+void FPyWrapperStructMetaData::AddReferencedObjects(FPyWrapperBase* Instance, FReferenceCollector& Collector)
+{
+	FPyWrapperStruct* Self = static_cast<FPyWrapperStruct*>(Instance);
+	Collector.AddReferencedObject(Self->ScriptStruct);
+	if (Self->ScriptStruct && Self->StructInstance && !Self->OwnerContext.HasOwner())
 	{
-		FPyWrapperStruct* StructSelf = static_cast<FPyWrapperStruct*>(Self);
-		Collector.AddReferencedObject(StructSelf->ScriptStruct);
-		if (StructSelf->ScriptStruct && StructSelf->StructInstance && !StructSelf->OwnerContext.HasOwner())
-		{
-			FPyReferenceCollector::AddReferencedObjectsFromStruct(Collector, StructSelf->ScriptStruct, StructSelf->StructInstance);
-		}
-	};
+		FPyReferenceCollector::AddReferencedObjectsFromStruct(Collector, Self->ScriptStruct, Self->StructInstance);
+	}
 }
 
 const IPyWrapperStructAllocationPolicy* FPyWrapperStructMetaData::GetAllocationPolicy(PyTypeObject* PyType)
@@ -617,9 +853,9 @@ struct FPythonGeneratedStructUtil
 		}
 		for (const TSharedPtr<PyGenUtil::FPropertyDef>& PropDef : InStruct->PropertyDefs)
 		{
-			PyGenUtil::FGeneratedWrappedMethodParameter& StructInitParam = StructInitParams[StructInitParams.AddDefaulted()];
+			PyGenUtil::FGeneratedWrappedMethodParameter& StructInitParam = StructInitParams.AddDefaulted_GetRef();
 			StructInitParam.ParamName = PropDef->GeneratedWrappedGetSet.GetSetName;
-			StructInitParam.ParamPropName = PropDef->GeneratedWrappedGetSet.PropName;
+			StructInitParam.ParamProp = PropDef->GeneratedWrappedGetSet.Prop;
 			StructInitParam.ParamDefaultValue = FString();
 		}
 
@@ -668,10 +904,10 @@ struct FPythonGeneratedStructUtil
 		InStruct->AddCppProperty(Prop);
 
 		// Build the definition data for the new property accessor
-		PyGenUtil::FPropertyDef& PropDef = *InStruct->PropertyDefs[InStruct->PropertyDefs.Add(MakeShared<PyGenUtil::FPropertyDef>())];
+		PyGenUtil::FPropertyDef& PropDef = *InStruct->PropertyDefs.Add_GetRef(MakeShared<PyGenUtil::FPropertyDef>());
 		PropDef.GeneratedWrappedGetSet.GetSetName = PyGenUtil::TCHARToUTF8Buffer(*InFieldName);
 		PropDef.GeneratedWrappedGetSet.GetSetDoc = PyGenUtil::TCHARToUTF8Buffer(*FString::Printf(TEXT("type: %s\n%s"), *PyGenUtil::GetPropertyPythonType(Prop), *PyGenUtil::GetFieldTooltip(Prop)));
-		PropDef.GeneratedWrappedGetSet.PropName = Prop->GetFName();
+		PropDef.GeneratedWrappedGetSet.Prop = Prop;
 		PropDef.GeneratedWrappedGetSet.GetCallback = (getter)&FPyWrapperStruct::Getter_Impl;
 		PropDef.GeneratedWrappedGetSet.SetCallback = (setter)&FPyWrapperStruct::Setter_Impl;
 		PropDef.GeneratedWrappedGetSet.ToPython(PropDef.PyGetSet);
diff --git a/Engine/Plugins/Experimental/PythonScriptPlugin/Source/PythonScriptPlugin/Private/PyWrapperStruct.h b/Engine/Plugins/Experimental/PythonScriptPlugin/Source/PythonScriptPlugin/Private/PyWrapperStruct.h
--- a/Engine/Plugins/Experimental/PythonScriptPlugin/Source/PythonScriptPlugin/Private/PyWrapperStruct.h
+++ b/Engine/Plugins/Experimental/PythonScriptPlugin/Source/PythonScriptPlugin/Private/PyWrapperStruct.h
@@ -50,19 +50,40 @@ struct FPyWrapperStruct : public FPyWrapperBase
 	static bool ValidateInternalState(FPyWrapperStruct* InSelf);
 
 	/** Cast the given Python object to this wrapped type (returns a new reference) */
-	static FPyWrapperStruct* CastPyObject(PyObject* InPyObject);
+	static FPyWrapperStruct* CastPyObject(PyObject* InPyObject, FPyConversionResult* OutCastResult = nullptr);
 
 	/** Cast the given Python object to this wrapped type, or attempt to convert the type into a new wrapped instance (returns a new reference) */
-	static FPyWrapperStruct* CastPyObject(PyObject* InPyObject, PyTypeObject* InType);
+	static FPyWrapperStruct* CastPyObject(PyObject* InPyObject, PyTypeObject* InType, FPyConversionResult* OutCastResult = nullptr);
 
 	/** Set the named property values on this instance (called via generated code) */
 	static int SetPropertyValues(FPyWrapperStruct* InSelf, PyObject* InArgs, PyObject* InKwds);
 
 	/** Get a named property value from this instance (called via generated code) */
-	static PyObject* GetPropertyValue(FPyWrapperStruct* InSelf, const FName InPropName, const char* InPythonAttrName);
+	static PyObject* GetPropertyValueByName(FPyWrapperStruct* InSelf, const FName InPropName, const char* InPythonAttrName);
+
+	/** Get a property value from this instance (called via generated code) */
+	static PyObject* GetPropertyValue(FPyWrapperStruct* InSelf, const UProperty* InProp, const char* InPythonAttrName);
 
 	/** Set a named property value on this instance (called via generated code) */
-	static int SetPropertyValue(FPyWrapperStruct* InSelf, PyObject* InValue, const FName InPropName, const char* InPythonAttrName, const bool InNotifyChange = false, const uint64 InReadOnlyFlags = CPF_EditConst | CPF_BlueprintReadOnly);
+	static int SetPropertyValueByName(FPyWrapperStruct* InSelf, PyObject* InValue, const FName InPropName, const char* InPythonAttrName, const bool InNotifyChange = false, const uint64 InReadOnlyFlags = CPF_EditConst | CPF_BlueprintReadOnly);
+
+	/** Set a property value on this instance (called via generated code) */
+	static int SetPropertyValue(FPyWrapperStruct* InSelf, PyObject* InValue, const UProperty* InProp, const char* InPythonAttrName, const bool InNotifyChange = false, const uint64 InReadOnlyFlags = CPF_EditConst | CPF_BlueprintReadOnly);
+
+	/** Call a function on this instance (CallFunction internal use only) */
+	static PyObject* CallFunction_Impl(FPyWrapperStruct* InSelf, PyObject* InArgs, PyObject* InKwds, const PyGenUtil::FGeneratedWrappedFunction& InFuncDef, const PyGenUtil::FGeneratedWrappedMethodParameter& InStructParam, const char* InPythonFuncName, const TCHAR* InErrorCtxt);
+
+	/** Implementation of the "call" logic for a Python method with no arguments (internal Python bindings use only) */
+	static PyObject* CallMethodNoArgs_Impl(FPyWrapperStruct* InSelf, void* InClosure);
+
+	/** Implementation of the "call" logic for a Python method with arguments (internal Python bindings use only) */
+	static PyObject* CallMethodWithArgs_Impl(FPyWrapperStruct* InSelf, PyObject* InArgs, PyObject* InKwds, void* InClosure);
+
+	/** Implementation of the "call" logic for a Python numeric binary function (internal CallBinaryOperator use only) */
+	static PyObject* CallBinaryOperatorFunction_Impl(FPyWrapperStruct* InSelf, PyObject* InRHS, const PyGenUtil::FGeneratedWrappedStructMathOpFunction& InMathOpFunc, const bool InInlineOp, const TOptional<EPyConversionResultState> InRequiredConversionResult = TOptional<EPyConversionResultState>(), FPyConversionResult* OutRHSConversionResult = nullptr);
+
+	/** Implementation of the "call" logic for a Python numeric binary function (internal Python bindings use only) */
+	static PyObject* CallBinaryOperator_Impl(FPyWrapperStruct* InSelf, PyObject* InRHS, const PyGenUtil::FGeneratedWrappedStructMathOpStack::EOpType InOpType);
 
 	/** Implementation of the "getter" logic for a Python descriptor reading from an struct property (internal Python bindings use only) */
 	static PyObject* Getter_Impl(FPyWrapperStruct* InSelf, void* InClosure);
@@ -95,15 +116,15 @@ struct TPyWrapperInlineStruct : public FPyWrapperStruct
 	TTypeCompatibleBytes<WrappedType> InlineStructInstance;
 
 	/** Cast the given Python object to this wrapped type (returns a new reference) */
-	static TPyWrapperInlineStruct* CastPyObject(PyObject* InPyObject)
+	static TPyWrapperInlineStruct* CastPyObject(PyObject* InPyObject, FPyConversionResult* OutCastResult = nullptr)
 	{
-		return (TPyWrapperInlineStruct*)FPyWrapperStruct::CastPyObject(InPyObject);
+		return (TPyWrapperInlineStruct*)FPyWrapperStruct::CastPyObject(InPyObject, OutCastResult);
 	}
 
 	/** Cast the given Python object to this wrapped type, or attempt to convert the type into a new wrapped instance (returns a new reference) */
-	static TPyWrapperInlineStruct* CastPyObject(PyObject* InPyObject, PyTypeObject* InType)
+	static TPyWrapperInlineStruct* CastPyObject(PyObject* InPyObject, PyTypeObject* InType, FPyConversionResult* OutCastResult = nullptr)
 	{
-		return (TPyWrapperInlineStruct*)FPyWrapperStruct::CastPyObject(InPyObject, InType);
+		return (TPyWrapperInlineStruct*)FPyWrapperStruct::CastPyObject(InPyObject, InType, OutCastResult);
 	}
 
 	/** Get a pointer to the typed struct instance this wrapper represents */
@@ -227,7 +248,7 @@ public:
 /** Meta-data for all UE4 exposed struct types */
 struct FPyWrapperStructMetaData : public FPyWrapperBaseMetaData
 {
-	PY_OVERRIDE_GETSET_METADATA(FPyWrapperStructMetaData)
+	PY_METADATA_METHODS(FPyWrapperStructMetaData, FGuid(0x03C9EA75, 0x2C86448B, 0xB53D1453, 0x94AA6BAE))
 
 	FPyWrapperStructMetaData();
 
@@ -249,6 +270,9 @@ struct FPyWrapperStructMetaData : public FPyWrapperBaseMetaData
 	/** Resolve the original property name of a Python method of the given instance */
 	static FName ResolvePropertyName(FPyWrapperStruct* Instance, const FName InPythonPropertyName);
 
+	/** Add object references from the given Python object to the given collector */
+	virtual void AddReferencedObjects(FPyWrapperBase* Instance, FReferenceCollector& Collector) override;
+
 	/** Allocation policy used with when allocating struct instances for a wrapped type */
 	const IPyWrapperStructAllocationPolicy* AllocPolicy;
 
@@ -260,6 +284,9 @@ struct FPyWrapperStructMetaData : public FPyWrapperBaseMetaData
 
 	/** Array of properties that were exposed to Python and can be used during init */
 	TArray<PyGenUtil::FGeneratedWrappedMethodParameter> InitParams;
+
+	/** The math operator stacks for this struct type */
+	PyGenUtil::FGeneratedWrappedStructMathOpStack MathOpStacks[(int32)PyGenUtil::FGeneratedWrappedStructMathOpStack::EOpType::Num];
 };
 
 /** Meta-data for all UE4 exposed inline struct types */
diff --git a/Engine/Plugins/Experimental/PythonScriptPlugin/Source/PythonScriptPlugin/Private/PyWrapperText.cpp b/Engine/Plugins/Experimental/PythonScriptPlugin/Source/PythonScriptPlugin/Private/PyWrapperText.cpp
--- a/Engine/Plugins/Experimental/PythonScriptPlugin/Source/PythonScriptPlugin/Private/PyWrapperText.cpp
+++ b/Engine/Plugins/Experimental/PythonScriptPlugin/Source/PythonScriptPlugin/Private/PyWrapperText.cpp
@@ -98,7 +98,7 @@ bool ExtractFormatArguments(FPyWrapperText* InSelf, PyObject* InObj, const int32
 						return false;
 					}
 
-					FFormatArgumentData& FormatArg = InOutFormatArgs[InOutFormatArgs.AddDefaulted()];
+					FFormatArgumentData& FormatArg = InOutFormatArgs.AddDefaulted_GetRef();
 					if (!ExtractFormatArgumentKey(InSelf, KeyItem, FormatArg))
 					{
 						PyUtil::SetPythonError(PyExc_TypeError, InSelf, *FString::Printf(TEXT("Cannot convert format argument %d (%s) at index %d"), InArgIndex, *PyUtil::GetFriendlyTypename(InObj), SequenceIndex));
@@ -122,7 +122,7 @@ bool ExtractFormatArguments(FPyWrapperText* InSelf, PyObject* InObj, const int32
 						return false;
 					}
 
-					FFormatArgumentData& FormatArg = InOutFormatArgs[InOutFormatArgs.AddDefaulted()];
+					FFormatArgumentData& FormatArg = InOutFormatArgs.AddDefaulted_GetRef();
 					FormatArg.ArgumentName = FString::FromInt(InArgIndex);
 					if (!ExtractFormatArgumentValue(InSelf, ValueItem, FormatArg))
 					{
@@ -135,7 +135,7 @@ bool ExtractFormatArguments(FPyWrapperText* InSelf, PyObject* InObj, const int32
 	}
 	else
 	{
-		FFormatArgumentData& FormatArg = InOutFormatArgs[InOutFormatArgs.AddDefaulted()];
+		FFormatArgumentData& FormatArg = InOutFormatArgs.AddDefaulted_GetRef();
 		FormatArg.ArgumentName = FString::FromInt(InArgIndex);
 		if (!ExtractFormatArgumentValue(InSelf, InObj, FormatArg))
 		{
@@ -158,63 +158,26 @@ void InitializePyWrapperText(PyObject* PyModule)
 	}
 }
 
-FPyWrapperText* FPyWrapperText::New(PyTypeObject* InType)
+void FPyWrapperText::InitValue(FPyWrapperText* InSelf, const FText InValue)
 {
-	FPyWrapperText* Self = (FPyWrapperText*)FPyWrapperBase::New(InType);
-	if (Self)
-	{
-		new(&Self->Text) FText();
-	}
-	return Self;
-}
-
-void FPyWrapperText::Free(FPyWrapperText* InSelf)
-{
-	Deinit(InSelf);
-
-	InSelf->Text.~FText();
-	FPyWrapperBase::Free(InSelf);
+	Super::InitValue(InSelf, InValue);
+	FPyWrapperTextFactory::Get().MapInstance(InSelf->Value, InSelf);
 }
 
-int FPyWrapperText::Init(FPyWrapperText* InSelf)
+void FPyWrapperText::DeinitValue(FPyWrapperText* InSelf)
 {
-	Deinit(InSelf);
-
-	const int BaseInit = FPyWrapperBase::Init(InSelf);
-	if (BaseInit != 0)
-	{
-		return BaseInit;
-	}
-
-	FPyWrapperTextFactory::Get().MapInstance(InSelf->Text, InSelf);
-	return 0;
-}
-
-int FPyWrapperText::Init(FPyWrapperText* InSelf, const FText InValue)
-{
-	Deinit(InSelf);
-
-	const int BaseInit = FPyWrapperBase::Init(InSelf);
-	if (BaseInit != 0)
-	{
-		return BaseInit;
-	}
-
-	InSelf->Text = InValue;
-	FPyWrapperTextFactory::Get().MapInstance(InSelf->Text, InSelf);
-	return 0;
+	FPyWrapperTextFactory::Get().UnmapInstance(InSelf->Value, Py_TYPE(InSelf));
+	Super::DeinitValue(InSelf);
 }
 
-void FPyWrapperText::Deinit(FPyWrapperText* InSelf)
+FPyWrapperText* FPyWrapperText::CastPyObject(PyObject* InPyObject, FPyConversionResult* OutCastResult)
 {
-	FPyWrapperTextFactory::Get().UnmapInstance(InSelf->Text, Py_TYPE(InSelf));
-	InSelf->Text = FText();
-}
+	SetOptionalPyConversionResult(FPyConversionResult::Failure(), OutCastResult);
 
-FPyWrapperText* FPyWrapperText::CastPyObject(PyObject* InPyObject)
-{
 	if (PyObject_IsInstance(InPyObject, (PyObject*)&PyWrapperTextType) == 1)
 	{
+		SetOptionalPyConversionResult(FPyConversionResult::Success(), OutCastResult);
+
 		Py_INCREF(InPyObject);
 		return (FPyWrapperText*)InPyObject;
 	}
@@ -222,10 +185,14 @@ FPyWrapperText* FPyWrapperText::CastPyObject(PyObject* InPyObject)
 	return nullptr;
 }
 
-FPyWrapperText* FPyWrapperText::CastPyObject(PyObject* InPyObject, PyTypeObject* InType)
+FPyWrapperText* FPyWrapperText::CastPyObject(PyObject* InPyObject, PyTypeObject* InType, FPyConversionResult* OutCastResult)
 {
+	SetOptionalPyConversionResult(FPyConversionResult::Failure(), OutCastResult);
+
 	if (PyObject_IsInstance(InPyObject, (PyObject*)InType) == 1 && (InType == &PyWrapperTextType || PyObject_IsInstance(InPyObject, (PyObject*)&PyWrapperTextType) == 1))
 	{
+		SetOptionalPyConversionResult(Py_TYPE(InPyObject) == InType ? FPyConversionResult::Success() : FPyConversionResult::SuccessWithCoercion(), OutCastResult);
+
 		Py_INCREF(InPyObject);
 		return (FPyWrapperText*)InPyObject;
 	}
@@ -242,6 +209,7 @@ FPyWrapperText* FPyWrapperText::CastPyObject(PyObject* InPyObject, PyTypeObject*
 					return nullptr;
 				}
 			}
+			SetOptionalPyConversionResult(FPyConversionResult::SuccessWithCoercion(), OutCastResult);
 			return NewText.Release();
 		}
 	}
@@ -253,16 +221,6 @@ PyTypeObject InitializePyWrapperTextType()
 {
 	struct FFuncs
 	{
-		static PyObject* New(PyTypeObject* InType, PyObject* InArgs, PyObject* InKwds)
-		{
-			return (PyObject*)FPyWrapperText::New(InType);
-		}
-
-		static void Dealloc(FPyWrapperText* InSelf)
-		{
-			FPyWrapperText::Free(InSelf);
-		}
-
 		static int Init(FPyWrapperText* InSelf, PyObject* InArgs, PyObject* InKwds)
 		{
 			FText InitValue;
@@ -282,7 +240,7 @@ PyTypeObject InitializePyWrapperTextType()
 
 		static PyObject* Str(FPyWrapperText* InSelf)
 		{
-			return PyUnicode_FromString(TCHAR_TO_UTF8(*InSelf->Text.ToString()));
+			return PyUnicode_FromString(TCHAR_TO_UTF8(*InSelf->Value.ToString()));
 		}
 
 		static PyObject* RichCmp(FPyWrapperText* InSelf, PyObject* InOther, int InOp)
@@ -294,7 +252,7 @@ PyTypeObject InitializePyWrapperTextType()
 				return Py_NotImplemented;
 			}
 
-			return PyUtil::PyRichCmp(InSelf->Text.CompareTo(Other), 0, InOp);
+			return PyUtil::PyRichCmp(InSelf->Value.CompareTo(Other), 0, InOp);
 		}
 
 		static PyUtil::FPyHashType Hash(FPyWrapperText* InSelf)
@@ -400,7 +358,7 @@ PyTypeObject InitializePyWrapperTextType()
 
 		static PyObject* IsEmpty(FPyWrapperText* InSelf)
 		{
-			if (InSelf->Text.IsEmpty())
+			if (InSelf->Value.IsEmpty())
 			{
 				Py_RETURN_TRUE;
 			}
@@ -410,7 +368,7 @@ PyTypeObject InitializePyWrapperTextType()
 
 		static PyObject* IsEmptyOrWhitespace(FPyWrapperText* InSelf)
 		{
-			if (InSelf->Text.IsEmptyOrWhitespace())
+			if (InSelf->Value.IsEmptyOrWhitespace())
 			{
 				Py_RETURN_TRUE;
 			}
@@ -420,7 +378,7 @@ PyTypeObject InitializePyWrapperTextType()
 
 		static PyObject* IsTransient(FPyWrapperText* InSelf)
 		{
-			if (InSelf->Text.IsTransient())
+			if (InSelf->Value.IsTransient())
 			{
 				Py_RETURN_TRUE;
 			}
@@ -430,7 +388,7 @@ PyTypeObject InitializePyWrapperTextType()
 
 		static PyObject* IsCultureInvariant(FPyWrapperText* InSelf)
 		{
-			if (InSelf->Text.IsCultureInvariant())
+			if (InSelf->Value.IsCultureInvariant())
 			{
 				Py_RETURN_TRUE;
 			}
@@ -440,7 +398,7 @@ PyTypeObject InitializePyWrapperTextType()
 
 		static PyObject* IsFromStringTable(FPyWrapperText* InSelf)
 		{
-			if (InSelf->Text.IsFromStringTable())
+			if (InSelf->Value.IsFromStringTable())
 			{
 				Py_RETURN_TRUE;
 			}
@@ -450,13 +408,13 @@ PyTypeObject InitializePyWrapperTextType()
 
 		static PyObject* ToLower(FPyWrapperText* InSelf)
 		{
-			const FText LowerText = InSelf->Text.ToLower();
+			const FText LowerText = InSelf->Value.ToLower();
 			return PyConversion::Pythonize(LowerText);
 		}
 
 		static PyObject* ToUpper(FPyWrapperText* InSelf)
 		{
-			const FText UpperText = InSelf->Text.ToUpper();
+			const FText UpperText = InSelf->Value.ToUpper();
 			return PyConversion::Pythonize(UpperText);
 		}
 
@@ -484,7 +442,7 @@ PyTypeObject InitializePyWrapperTextType()
 				return nullptr;
 			}
 
-			const FText FormattedText = FTextFormatter::Format(InSelf->Text, MoveTemp(FormatArgs), false, false);
+			const FText FormattedText = FTextFormatter::Format(InSelf->Value, MoveTemp(FormatArgs), false, false);
 			return PyConversion::Pythonize(FormattedText);
 		}
 	};
@@ -506,15 +464,8 @@ PyTypeObject InitializePyWrapperTextType()
 		{ nullptr, nullptr, 0, nullptr }
 	};
 
-	PyTypeObject PyType = {
-		PyVarObject_HEAD_INIT(nullptr, 0)
-		"Text", /* tp_name */
-		sizeof(FPyWrapperText), /* tp_basicsize */
-	};
+	PyTypeObject PyType = InitializePyWrapperBasicType<FPyWrapperText>("Text", "Type for all UE4 exposed FText instances");
 
-	PyType.tp_base = &PyWrapperBaseType;
-	PyType.tp_new = (newfunc)&FFuncs::New;
-	PyType.tp_dealloc = (destructor)&FFuncs::Dealloc;
 	PyType.tp_init = (initproc)&FFuncs::Init;
 	PyType.tp_str = (reprfunc)&FFuncs::Str;
 	PyType.tp_richcompare = (richcmpfunc)&FFuncs::RichCmp;
@@ -522,9 +473,6 @@ PyTypeObject InitializePyWrapperTextType()
 
 	PyType.tp_methods = PyMethods;
 
-	PyType.tp_flags = Py_TPFLAGS_DEFAULT;
-	PyType.tp_doc = "Type for all UE4 exposed FText instances";
-
 	return PyType;
 }
 
diff --git a/Engine/Plugins/Experimental/PythonScriptPlugin/Source/PythonScriptPlugin/Private/PyWrapperText.h b/Engine/Plugins/Experimental/PythonScriptPlugin/Source/PythonScriptPlugin/Private/PyWrapperText.h
--- a/Engine/Plugins/Experimental/PythonScriptPlugin/Source/PythonScriptPlugin/Private/PyWrapperText.h
+++ b/Engine/Plugins/Experimental/PythonScriptPlugin/Source/PythonScriptPlugin/Private/PyWrapperText.h
@@ -2,7 +2,7 @@
 
 #pragma once
 
-#include "PyWrapperBase.h"
+#include "PyWrapperBasic.h"
 
 #if WITH_PYTHON
 
@@ -13,31 +13,21 @@ extern PyTypeObject PyWrapperTextType;
 void InitializePyWrapperText(PyObject* PyModule);
 
 /** Type for all UE4 exposed FText instances */
-struct FPyWrapperText : public FPyWrapperBase
+struct FPyWrapperText : public TPyWrapperBasic<FText, FPyWrapperText>
 {
-	/** The wrapped text */
-	FText Text;
-
-	/** New this wrapper instance (called via tp_new for Python, or directly in C++) */
-	static FPyWrapperText* New(PyTypeObject* InType);
-
-	/** Free this wrapper instance (called via tp_dealloc for Python) */
-	static void Free(FPyWrapperText* InSelf);
-
-	/** Initialize this wrapper instance (called via tp_init for Python, or directly in C++) */
-	static int Init(FPyWrapperText* InSelf);
-
-	/** Initialize this wrapper instance to the given value (called via tp_init for Python, or directly in C++) */
-	static int Init(FPyWrapperText* InSelf, const FText InValue);
-
-	/** Deinitialize this wrapper instance (called via Init and Free to restore the instance to its New state) */
-	static void Deinit(FPyWrapperText* InSelf);
+	typedef TPyWrapperBasic<FText, FPyWrapperText> Super;
 
 	/** Cast the given Python object to this wrapped type (returns a new reference) */
-	static FPyWrapperText* CastPyObject(PyObject* InPyObject);
+	static FPyWrapperText* CastPyObject(PyObject* InPyObject, FPyConversionResult* OutCastResult = nullptr);
 
 	/** Cast the given Python object to this wrapped type, or attempt to convert the type into a new wrapped instance (returns a new reference) */
-	static FPyWrapperText* CastPyObject(PyObject* InPyObject, PyTypeObject* InType);
+	static FPyWrapperText* CastPyObject(PyObject* InPyObject, PyTypeObject* InType, FPyConversionResult* OutCastResult = nullptr);
+
+	/** Initialize the value of this wrapper instance (internal) */
+	static void InitValue(FPyWrapperText* InSelf, const FText InValue);
+
+	/** Deinitialize the value of this wrapper instance (internal) */
+	static void DeinitValue(FPyWrapperText* InSelf);
 };
 
 typedef TPyPtr<FPyWrapperText> FPyWrapperTextPtr;
diff --git a/Engine/Plugins/Experimental/PythonScriptPlugin/Source/PythonScriptPlugin/Private/PyWrapperTypeRegistry.cpp b/Engine/Plugins/Experimental/PythonScriptPlugin/Source/PythonScriptPlugin/Private/PyWrapperTypeRegistry.cpp
--- a/Engine/Plugins/Experimental/PythonScriptPlugin/Source/PythonScriptPlugin/Private/PyWrapperTypeRegistry.cpp
+++ b/Engine/Plugins/Experimental/PythonScriptPlugin/Source/PythonScriptPlugin/Private/PyWrapperTypeRegistry.cpp
@@ -23,6 +23,22 @@
 
 #if WITH_PYTHON
 
+DECLARE_FLOAT_ACCUMULATOR_STAT(TEXT("Generate Wrapped Class Total Time"), STAT_GenerateWrappedClassTotalTime, STATGROUP_Python);
+DECLARE_DWORD_ACCUMULATOR_STAT(TEXT("Generate Wrapped Class Call Count"), STAT_GenerateWrappedClassCallCount, STATGROUP_Python);
+DECLARE_DWORD_ACCUMULATOR_STAT(TEXT("Generate Wrapped Class Obj Count"), STAT_GenerateWrappedClassObjCount, STATGROUP_Python);
+
+DECLARE_FLOAT_ACCUMULATOR_STAT(TEXT("Generate Wrapped Struct Total Time"), STAT_GenerateWrappedStructTotalTime, STATGROUP_Python);
+DECLARE_DWORD_ACCUMULATOR_STAT(TEXT("Generate Wrapped Struct Call Count"), STAT_GenerateWrappedStructCallCount, STATGROUP_Python);
+DECLARE_DWORD_ACCUMULATOR_STAT(TEXT("Generate Wrapped Struct Obj Count"), STAT_GenerateWrappedStructObjCount, STATGROUP_Python);
+
+DECLARE_FLOAT_ACCUMULATOR_STAT(TEXT("Generate Wrapped Enum Total Time"), STAT_GenerateWrappedEnumTotalTime, STATGROUP_Python);
+DECLARE_DWORD_ACCUMULATOR_STAT(TEXT("Generate Wrapped Enum Call Count"), STAT_GenerateWrappedEnumCallCount, STATGROUP_Python);
+DECLARE_DWORD_ACCUMULATOR_STAT(TEXT("Generate Wrapped Enum Obj Count"), STAT_GenerateWrappedEnumObjCount, STATGROUP_Python);
+
+DECLARE_FLOAT_ACCUMULATOR_STAT(TEXT("Generate Wrapped Delegate Total Time"), STAT_GenerateWrappedDelegateTotalTime, STATGROUP_Python);
+DECLARE_DWORD_ACCUMULATOR_STAT(TEXT("Generate Wrapped Delegate Call Count"), STAT_GenerateWrappedDelegateCallCount, STATGROUP_Python);
+DECLARE_DWORD_ACCUMULATOR_STAT(TEXT("Generate Wrapped Delegate Obj Count"), STAT_GenerateWrappedDelegateObjCount, STATGROUP_Python);
+
 FPyWrapperObjectFactory& FPyWrapperObjectFactory::Get()
 {
 	static FPyWrapperObjectFactory Instance;
@@ -126,9 +142,9 @@ FPyWrapperDelegate* FPyWrapperDelegateFactory::CreateInstance(const UFunction* I
 	}
 
 	PyTypeObject* PyType = FPyWrapperTypeRegistry::Get().GetWrappedDelegateType(InDelegateSignature);
-	return CreateInstanceInternal(InUnrealInstance, PyType, [InDelegateSignature, InUnrealInstance, &InOwnerContext, InConversionMethod](FPyWrapperDelegate* InSelf)
+	return CreateInstanceInternal(InUnrealInstance, PyType, [InUnrealInstance, &InOwnerContext, InConversionMethod](FPyWrapperDelegate* InSelf)
 	{
-		return FPyWrapperDelegate::Init(InSelf, InOwnerContext, InDelegateSignature, InUnrealInstance, InConversionMethod);
+		return FPyWrapperDelegate::Init(InSelf, InOwnerContext, InUnrealInstance, InConversionMethod);
 	}, InConversionMethod == EPyConversionMethod::Copy || InConversionMethod == EPyConversionMethod::Steal);
 }
 
@@ -158,9 +174,9 @@ FPyWrapperMulticastDelegate* FPyWrapperMulticastDelegateFactory::CreateInstance(
 	}
 
 	PyTypeObject* PyType = FPyWrapperTypeRegistry::Get().GetWrappedDelegateType(InDelegateSignature);
-	return CreateInstanceInternal(InUnrealInstance, PyType, [InDelegateSignature, InUnrealInstance, &InOwnerContext, InConversionMethod](FPyWrapperMulticastDelegate* InSelf)
+	return CreateInstanceInternal(InUnrealInstance, PyType, [InUnrealInstance, &InOwnerContext, InConversionMethod](FPyWrapperMulticastDelegate* InSelf)
 	{
-		return FPyWrapperMulticastDelegate::Init(InSelf, InOwnerContext, InDelegateSignature, InUnrealInstance, InConversionMethod);
+		return FPyWrapperMulticastDelegate::Init(InSelf, InOwnerContext, InUnrealInstance, InConversionMethod);
 	}, InConversionMethod == EPyConversionMethod::Copy || InConversionMethod == EPyConversionMethod::Steal);
 }
 
@@ -506,6 +522,9 @@ PyTypeObject* FPyWrapperTypeRegistry::GenerateWrappedTypeForObject(const UObject
 
 PyTypeObject* FPyWrapperTypeRegistry::GenerateWrappedClassType(const UClass* InClass, FGeneratedWrappedTypeReferences& OutGeneratedWrappedTypeReferences, TSet<FName>& OutDirtyModules, const bool bForceGenerate)
 {
+	SCOPE_SECONDS_ACCUMULATOR(STAT_GenerateWrappedClassTotalTime);
+	INC_DWORD_STAT(STAT_GenerateWrappedClassCallCount);
+
 	// Already processed? Nothing more to do
 	if (PyTypeObject* ExistingPyType = PythonWrappedClasses.FindRef(InClass->GetFName()))
 	{
@@ -531,8 +550,11 @@ PyTypeObject* FPyWrapperTypeRegistry::GenerateWrappedClassType(const UClass* InC
 		SuperPyType = GenerateWrappedClassType(SuperClass, OutGeneratedWrappedTypeReferences, OutDirtyModules, true);
 	}
 
+	INC_DWORD_STAT(STAT_GenerateWrappedClassObjCount);
+
 	check(!GeneratedWrappedTypes.Contains(InClass->GetFName()));
-	TSharedPtr<PyGenUtil::FGeneratedWrappedType> GeneratedWrappedType = GeneratedWrappedTypes.Add(InClass->GetFName(), MakeShared<PyGenUtil::FGeneratedWrappedType>());
+	TSharedRef<PyGenUtil::FGeneratedWrappedClassType> GeneratedWrappedType = MakeShared<PyGenUtil::FGeneratedWrappedClassType>();
+	GeneratedWrappedTypes.Add(InClass->GetFName(), GeneratedWrappedType);
 
 	TMap<FName, FName> PythonProperties;
 	TMap<FName, FName> PythonMethods;
@@ -546,138 +568,190 @@ PyTypeObject* FPyWrapperTypeRegistry::GenerateWrappedClassType(const UClass* InC
 		{
 			GatherWrappedTypesForPropertyReferences(InProp, OutGeneratedWrappedTypeReferences);
 
-			const PyGenUtil::FGeneratedWrappedPropertyDoc& GeneratedPropertyDoc = GeneratedWrappedType->TypePropertyDocs[GeneratedWrappedType->TypePropertyDocs.Emplace(InProp)];
+			const PyGenUtil::FGeneratedWrappedPropertyDoc& GeneratedPropertyDoc = GeneratedWrappedType->PropertyDocs[GeneratedWrappedType->PropertyDocs.Emplace(InProp)];
 			PythonProperties.Add(*GeneratedPropertyDoc.PythonPropName, InProp->GetFName());
 
 			if (bExportPropertyToScript)
 			{
-				PyGenUtil::FGeneratedWrappedGetSet& GeneratedWrappedGetSet = GeneratedWrappedType->TypeGetSets[GeneratedWrappedType->TypeGetSets.AddDefaulted()];
-				GeneratedWrappedGetSet.Class = (UClass*)InClass;
+				PyGenUtil::FGeneratedWrappedGetSet& GeneratedWrappedGetSet = GeneratedWrappedType->GetSets.TypeGetSets.AddDefaulted_GetRef();
 				GeneratedWrappedGetSet.GetSetName = PyGenUtil::TCHARToUTF8Buffer(*GeneratedPropertyDoc.PythonPropName);
 				GeneratedWrappedGetSet.GetSetDoc = PyGenUtil::TCHARToUTF8Buffer(*GeneratedPropertyDoc.DocString);
-				GeneratedWrappedGetSet.PropName = InProp->GetFName();
-				GeneratedWrappedGetSet.GetFuncName = *InProp->GetMetaData(PyGenUtil::BlueprintGetterMetaDataKey);
-				GeneratedWrappedGetSet.SetFuncName = *InProp->GetMetaData(PyGenUtil::BlueprintSetterMetaDataKey);
+				GeneratedWrappedGetSet.Prop = InProp;
+				GeneratedWrappedGetSet.GetFunc.SetFunctionAndExtractParams(InClass->FindFunctionByName(*InProp->GetMetaData(PyGenUtil::BlueprintGetterMetaDataKey)));
+				GeneratedWrappedGetSet.SetFunc.SetFunctionAndExtractParams(InClass->FindFunctionByName(*InProp->GetMetaData(PyGenUtil::BlueprintSetterMetaDataKey)));
 				GeneratedWrappedGetSet.GetCallback = (getter)&FPyWrapperObject::Getter_Impl;
 				GeneratedWrappedGetSet.SetCallback = (setter)&FPyWrapperObject::Setter_Impl;
 			}
 		}
 	};
 
-	auto GenerateWrappedMethod = [this, InClass, &PythonMethods, &GeneratedWrappedType, &OutGeneratedWrappedTypeReferences](const UFunction* InFunc)
+	auto GenerateDynamicStructMethod = [this, &OutGeneratedWrappedTypeReferences, &OutDirtyModules](const UFunction* InFunc, const PyGenUtil::FGeneratedWrappedMethod& InTypeMethod)
 	{
-		if (!PyGenUtil::ShouldExportFunction(InFunc))
+		// Only static functions can be hoisted into the struct
+		if (!InFunc->HasAnyFunctionFlags(FUNC_Static))
 		{
+			UE_LOG(LogPython, Warning, TEXT("Non-static function '%s' is marked as 'ScriptMethod' but only static functions can be hoisted."), *InFunc->GetName());
 			return;
 		}
 
-		const FString PythonFunctionName = PyGenUtil::GetFunctionPythonName(InFunc);
-		const bool bIsStatic = InFunc->HasAnyFunctionFlags(FUNC_Static);
+		// Get the struct type to hoist this method to (this should be the first parameter)
+		PyGenUtil::FGeneratedWrappedMethodParameter StructParam;
+		if (InTypeMethod.MethodFunc.InputParams.Num() > 0 && InTypeMethod.MethodFunc.InputParams[0].ParamProp->IsA<UStructProperty>())
+		{
+			StructParam = InTypeMethod.MethodFunc.InputParams[0];
+		}
+		if (!StructParam.ParamProp)
+		{
+			UE_LOG(LogPython, Warning, TEXT("Function '%s' is marked as 'ScriptMethod' but doesn't contain a valid struct as its first argument."), *InFunc->GetName());
+			return;
+		}
+
+		// Ensure that we've generated a finalized Python type for this struct since we'll be adding this function as a dynamic method on that type
+		const UStruct* Struct = CastChecked<UStructProperty>(StructParam.ParamProp)->Struct;
+		if (!GenerateWrappedStructType(Struct, OutGeneratedWrappedTypeReferences, OutDirtyModules, true))
+		{
+			return;
+		}
+
+		// Find the wrapped type for the struct as that's what we'll actually add the dynamic method to
+		TSharedPtr<PyGenUtil::FGeneratedWrappedStructType> StructGeneratedWrappedType = StaticCastSharedPtr<PyGenUtil::FGeneratedWrappedStructType>(GeneratedWrappedTypes.FindRef(Struct->GetFName()));
+		check(StructGeneratedWrappedType.IsValid());
+
+		// Copy the basic wrapped method as we need to adjust some parts of it below
+		PyGenUtil::FGeneratedWrappedDynamicStructMethod GeneratedWrappedDynamicStructMethod;
+		static_cast<PyGenUtil::FGeneratedWrappedMethod&>(GeneratedWrappedDynamicStructMethod) = InTypeMethod;
+		GeneratedWrappedDynamicStructMethod.StructParam = StructParam;
+
+		// Hoisted methods may use an optional name alias
+		FString PythonStructMethodName = InFunc->GetMetaData(PyGenUtil::ScriptMethodMetaDataKey);
+		if (PythonStructMethodName.IsEmpty())
+		{
+			PythonStructMethodName = UTF8_TO_TCHAR(InTypeMethod.MethodName.GetData());
+		}
+		else
+		{
+			PythonStructMethodName = PythonizeName(PythonStructMethodName, PyGenUtil::EPythonizeNameCase::Lower);
+		}
+		GeneratedWrappedDynamicStructMethod.MethodName = PyGenUtil::TCHARToUTF8Buffer(*PythonStructMethodName);
 		
-		PythonMethods.Add(*PythonFunctionName, InFunc->GetFName());
+		// We remove the first function parameter, as that's the struct and we'll infer that when we call
+		GeneratedWrappedDynamicStructMethod.MethodFunc.InputParams.RemoveAt(0, 1, /*bAllowShrinking*/false);
 
-		PyGenUtil::FGeneratedWrappedMethod& GeneratedWrappedMethod = GeneratedWrappedType->TypeMethods[GeneratedWrappedType->TypeMethods.AddDefaulted()];
-		GeneratedWrappedMethod.Class = (UClass*)InClass;
-		GeneratedWrappedMethod.MethodName = PyGenUtil::TCHARToUTF8Buffer(*PythonFunctionName);
-		GeneratedWrappedMethod.MethodFuncName = InFunc->GetFName();
+		// Set-up some data needed to build the tooltip correctly for the hoisted method
+		const bool bIsStaticOverride = false;
+		TSet<FName> ParamsToIgnore;
+		ParamsToIgnore.Add(StructParam.ParamProp->GetFName());
+
+		// Update the doc string for the method
+		FString PythonStructMethodDocString = PyGenUtil::BuildFunctionDocString(InFunc, PythonStructMethodName, GeneratedWrappedDynamicStructMethod.MethodFunc.InputParams, GeneratedWrappedDynamicStructMethod.MethodFunc.OutputParams, &bIsStaticOverride);
+		PythonStructMethodDocString += TEXT(" -- ");
+		PythonStructMethodDocString += PyGenUtil::PythonizeFunctionTooltip(PyGenUtil::GetFieldTooltip(InFunc), InFunc, ParamsToIgnore);
+		GeneratedWrappedDynamicStructMethod.MethodDoc = PyGenUtil::TCHARToUTF8Buffer(*PythonStructMethodDocString);
+
+		// Update the flags as removing the struct argument may have changed the calling convention
+		GeneratedWrappedDynamicStructMethod.MethodFlags = GeneratedWrappedDynamicStructMethod.MethodFunc.InputParams.Num() > 0 ? METH_VARARGS | METH_KEYWORDS : METH_NOARGS;
+
+		// Set the correct function pointer for calling this function and inject the struct argument
+		GeneratedWrappedDynamicStructMethod.MethodCallback = GeneratedWrappedDynamicStructMethod.MethodFunc.InputParams.Num() > 0 ? PyCFunctionWithClosureCast(&FPyWrapperStruct::CallMethodWithArgs_Impl) : PyCFunctionWithClosureCast(&FPyWrapperStruct::CallMethodNoArgs_Impl);
+
+		// Finished, add the dynamic method now
+		StructGeneratedWrappedType->AddDynamicStructMethod(MoveTemp(GeneratedWrappedDynamicStructMethod));
+	};
 
-		TArray<const UProperty*, TInlineAllocator<4>> OutParams;
-		if (const UProperty* ReturnProp = InFunc->GetReturnProperty())
+	auto GenerateStructMathOp = [this, &OutGeneratedWrappedTypeReferences, &OutDirtyModules](const UFunction* InFunc, const PyGenUtil::FGeneratedWrappedMethod& InTypeMethod)
+	{
+		// Only static functions can be hoisted into the struct
+		if (!InFunc->HasAnyFunctionFlags(FUNC_Static))
 		{
-			OutParams.Add(ReturnProp);
+			UE_LOG(LogPython, Warning, TEXT("Non-static function '%s' is marked as 'ScriptMathOp' but only static functions can be hoisted."), *InFunc->GetName());
+			return;
 		}
 
-		FString FunctionDeclDocString = FString::Printf(TEXT("%s.%s("), (bIsStatic ? TEXT("X") : TEXT("x")), *PythonFunctionName);
-		for (TFieldIterator<const UProperty> ParamIt(InFunc); ParamIt; ++ParamIt)
+		// Get the struct type to hoist this method to (this should be the first parameter)
+		PyGenUtil::FGeneratedWrappedStructMathOpFunction MathOpFunc;
+		static_cast<PyGenUtil::FGeneratedWrappedFunction&>(MathOpFunc) = InTypeMethod.MethodFunc;
+		if (MathOpFunc.InputParams.Num() > 0 && MathOpFunc.InputParams[0].ParamProp->IsA<UStructProperty>())
 		{
-			const UProperty* Param = *ParamIt;
+			MathOpFunc.StructParam = MathOpFunc.InputParams[0];
+			MathOpFunc.InputParams.RemoveAt(0, 1, /*bAllowShrinking*/false);
+		}
+		if (!MathOpFunc.StructParam.ParamProp)
+		{
+			UE_LOG(LogPython, Warning, TEXT("Function '%s' is marked as 'ScriptMathOp' but doesn't contain a valid struct as its first argument."), *InFunc->GetName());
+			return;
+		}
 
-			if (PyUtil::IsInputParameter(Param))
-			{
-				const FString ParamName = Param->GetName();
-				const FString PythonParamName = PyGenUtil::PythonizePropertyName(ParamName, PyGenUtil::EPythonizeNameCase::Lower);
-				const FName DefaultValueMetaDataKey = *FString::Printf(TEXT("CPP_Default_%s"), *ParamName);
-
-				PyGenUtil::FGeneratedWrappedMethodParameter& GeneratedWrappedMethodParam = GeneratedWrappedMethod.MethodParams[GeneratedWrappedMethod.MethodParams.AddDefaulted()];
-				GeneratedWrappedMethodParam.ParamName = PyGenUtil::TCHARToUTF8Buffer(*PythonParamName);
-				GeneratedWrappedMethodParam.ParamPropName = Param->GetFName();
-				if (InFunc->HasMetaData(DefaultValueMetaDataKey))
-				{
-					GeneratedWrappedMethodParam.ParamDefaultValue = InFunc->GetMetaData(DefaultValueMetaDataKey);
-				}
-
-				if (FunctionDeclDocString[FunctionDeclDocString.Len() - 1] != TEXT('('))
-				{
-					FunctionDeclDocString += TEXT(", ");
-				}
-				FunctionDeclDocString += PythonParamName;
-				if (GeneratedWrappedMethodParam.ParamDefaultValue.IsSet())
-				{
-					FunctionDeclDocString += TEXT('=');
-					FunctionDeclDocString += GeneratedWrappedMethodParam.ParamDefaultValue.GetValue();
-				}
-			}
+		// Ensure that we've generated a finalized Python type for this struct since we'll be adding this function as a operator on that type
+		const UStruct* Struct = CastChecked<UStructProperty>(MathOpFunc.StructParam.ParamProp)->Struct;
+		if (!GenerateWrappedStructType(Struct, OutGeneratedWrappedTypeReferences, OutDirtyModules, true))
+		{
+			return;
+		}
 
-			if (PyUtil::IsOutputParameter(Param))
+		// Find the wrapped type for the struct as that's what we'll actually add the operator to (via its meta-data)
+		TSharedPtr<PyGenUtil::FGeneratedWrappedStructType> StructGeneratedWrappedType = StaticCastSharedPtr<PyGenUtil::FGeneratedWrappedStructType>(GeneratedWrappedTypes.FindRef(Struct->GetFName()));
+		check(StructGeneratedWrappedType.IsValid());
+
+		// Add the function to each specified operator
+		const FString& MathOpsStr = InFunc->GetMetaData(PyGenUtil::ScriptMathOpMetaDataKey);
+		TArray<FString> MathOpStrArray;
+		MathOpsStr.ParseIntoArray(MathOpStrArray, TEXT(";"));
+		for (const FString& MathOpStr : MathOpStrArray)
+		{
+			PyGenUtil::FGeneratedWrappedStructMathOpStack::EOpType MathOp;
+			if (PyGenUtil::FGeneratedWrappedStructMathOpStack::StringToOpType(*MathOpStr, MathOp))
 			{
-				OutParams.Add(Param);
+				StaticCastSharedPtr<FPyWrapperStructMetaData>(StructGeneratedWrappedType->MetaData)->MathOpStacks[(int32)MathOp].MathOpFuncs.Add(MathOpFunc);
 			}
-
-			GatherWrappedTypesForPropertyReferences(Param, OutGeneratedWrappedTypeReferences);
 		}
-		FunctionDeclDocString += TEXT(")");
+	};
 
-		if (OutParams.Num() > 0)
+	auto GenerateWrappedMethod = [this, &PythonMethods, &GeneratedWrappedType, &OutGeneratedWrappedTypeReferences, &GenerateDynamicStructMethod, &GenerateStructMathOp](const UFunction* InFunc)
+	{
+		if (!PyGenUtil::ShouldExportFunction(InFunc))
 		{
-			FunctionDeclDocString += TEXT(" -> ");
+			return;
+		}
 
-			// If we have multiple return values and the main return value is a bool, we return None (for false) or the (potentially packed) return value without the bool (for true)
-			bool bStrippedBoolReturn = false;
-			if (OutParams.Num() > 1 && OutParams[0]->IsA<UBoolProperty>())
-			{
-				OutParams.RemoveAt(0, 1, false);
-				bStrippedBoolReturn = true;
-			}
+		const FString PythonFunctionName = PyGenUtil::GetFunctionPythonName(InFunc);
+		const bool bIsStatic = InFunc->HasAnyFunctionFlags(FUNC_Static);
+		
+		PythonMethods.Add(*PythonFunctionName, InFunc->GetFName());
 
-			if (OutParams.Num() == 1)
-			{
-				FunctionDeclDocString += PyGenUtil::GetPropertyTypePythonName(OutParams[0]);
-			}
-			else
-			{
-				FunctionDeclDocString += TEXT('(');
-				for (int32 OutParamIndex = 0; OutParamIndex < OutParams.Num(); ++OutParamIndex)
-				{
-					if (OutParamIndex > 0)
-					{
-						FunctionDeclDocString += TEXT(", ");
-					}
-					if (OutParamIndex > 0 || bStrippedBoolReturn)
-					{
-						FunctionDeclDocString += PyGenUtil::PythonizePropertyName(OutParams[OutParamIndex]->GetName(), PyGenUtil::EPythonizeNameCase::Lower);
-						FunctionDeclDocString += TEXT('=');
-					}
-					FunctionDeclDocString += PyGenUtil::GetPropertyTypePythonName(OutParams[OutParamIndex]);
-				}
-				FunctionDeclDocString += TEXT(')');
-			}
+		PyGenUtil::FGeneratedWrappedMethod& GeneratedWrappedMethod = GeneratedWrappedType->Methods.TypeMethods.AddDefaulted_GetRef();
+		GeneratedWrappedMethod.MethodName = PyGenUtil::TCHARToUTF8Buffer(*PythonFunctionName);
+		GeneratedWrappedMethod.MethodFunc.SetFunctionAndExtractParams(InFunc);
 
-			if (bStrippedBoolReturn)
-			{
-				FunctionDeclDocString += TEXT(" or None");
-			}
+		for (TFieldIterator<const UProperty> ParamIt(InFunc); ParamIt; ++ParamIt)
+		{
+			const UProperty* Param = *ParamIt;
+			GatherWrappedTypesForPropertyReferences(Param, OutGeneratedWrappedTypeReferences);
 		}
 
-		GeneratedWrappedMethod.MethodDoc = PyGenUtil::TCHARToUTF8Buffer(*FString::Printf(TEXT("%s -- %s"), *FunctionDeclDocString, *PyGenUtil::PythonizeFunctionTooltip(PyGenUtil::GetFieldTooltip(InFunc), InFunc)));
-		GeneratedWrappedMethod.MethodFlags = GeneratedWrappedMethod.MethodParams.Num() > 0 ? METH_VARARGS | METH_KEYWORDS : METH_NOARGS;
+		FString FunctionDeclDocString = PyGenUtil::BuildFunctionDocString(InFunc, PythonFunctionName, GeneratedWrappedMethod.MethodFunc.InputParams, GeneratedWrappedMethod.MethodFunc.OutputParams);
+		FunctionDeclDocString += TEXT(" -- ");
+		FunctionDeclDocString += PyGenUtil::PythonizeFunctionTooltip(PyGenUtil::GetFieldTooltip(InFunc), InFunc);
+
+		GeneratedWrappedMethod.MethodDoc = PyGenUtil::TCHARToUTF8Buffer(*FunctionDeclDocString);
+		GeneratedWrappedMethod.MethodFlags = GeneratedWrappedMethod.MethodFunc.InputParams.Num() > 0 ? METH_VARARGS | METH_KEYWORDS : METH_NOARGS;
 		if (bIsStatic)
 		{
 			GeneratedWrappedMethod.MethodFlags |= METH_CLASS;
-			GeneratedWrappedMethod.MethodCallback = GeneratedWrappedMethod.MethodParams.Num() > 0 ? PyCFunctionWithClosureCast(&FPyWrapperObject::CallClassMethodWithArgs_Impl) : PyCFunctionWithClosureCast(&FPyWrapperObject::CallClassMethodNoArgs_Impl);
+			GeneratedWrappedMethod.MethodCallback = GeneratedWrappedMethod.MethodFunc.InputParams.Num() > 0 ? PyCFunctionWithClosureCast(&FPyWrapperObject::CallClassMethodWithArgs_Impl) : PyCFunctionWithClosureCast(&FPyWrapperObject::CallClassMethodNoArgs_Impl);
 		}
 		else
 		{
-			GeneratedWrappedMethod.MethodCallback = GeneratedWrappedMethod.MethodParams.Num() > 0 ? PyCFunctionWithClosureCast(&FPyWrapperObject::CallMethodWithArgs_Impl) : PyCFunctionWithClosureCast(&FPyWrapperObject::CallMethodNoArgs_Impl);
+			GeneratedWrappedMethod.MethodCallback = GeneratedWrappedMethod.MethodFunc.InputParams.Num() > 0 ? PyCFunctionWithClosureCast(&FPyWrapperObject::CallMethodWithArgs_Impl) : PyCFunctionWithClosureCast(&FPyWrapperObject::CallMethodNoArgs_Impl);
+		}
+
+		// Should this function also be hoisted as a struct method or math op?
+		if (InFunc->HasMetaData(PyGenUtil::ScriptMethodMetaDataKey))
+		{
+			GenerateDynamicStructMethod(InFunc, GeneratedWrappedMethod);
+		}
+		if (InFunc->HasMetaData(PyGenUtil::ScriptMathOpMetaDataKey))
+		{
+			GenerateStructMathOp(InFunc, GeneratedWrappedMethod);
 		}
 	};
 
@@ -701,21 +775,21 @@ PyTypeObject* FPyWrapperTypeRegistry::GenerateWrappedClassType(const UClass* InC
 	FString TypeDocString = PyGenUtil::PythonizeTooltip(PyGenUtil::GetFieldTooltip(InClass));
 	if (const UClass* SuperClass = InClass->GetSuperClass())
 	{
-		TSharedPtr<PyGenUtil::FGeneratedWrappedType> SuperGeneratedWrappedType = GeneratedWrappedTypes.FindRef(SuperClass->GetFName());
+		TSharedPtr<PyGenUtil::FGeneratedWrappedClassType> SuperGeneratedWrappedType = StaticCastSharedPtr<PyGenUtil::FGeneratedWrappedClassType>(GeneratedWrappedTypes.FindRef(SuperClass->GetFName()));
 		if (SuperGeneratedWrappedType.IsValid())
 		{
-			GeneratedWrappedType->TypePropertyDocs.Append(SuperGeneratedWrappedType->TypePropertyDocs);
+			GeneratedWrappedType->PropertyDocs.Append(SuperGeneratedWrappedType->PropertyDocs);
 		}
 	}
-	GeneratedWrappedType->TypePropertyDocs.Sort(&PyGenUtil::FGeneratedWrappedPropertyDoc::SortPredicate);
-	PyGenUtil::FGeneratedWrappedPropertyDoc::AppendDocString(GeneratedWrappedType->TypePropertyDocs, TypeDocString, /*bEditorVariant*/true);
+	GeneratedWrappedType->PropertyDocs.Sort(&PyGenUtil::FGeneratedWrappedPropertyDoc::SortPredicate);
+	PyGenUtil::FGeneratedWrappedPropertyDoc::AppendDocString(GeneratedWrappedType->PropertyDocs, TypeDocString, /*bEditorVariant*/true);
 	GeneratedWrappedType->TypeDoc = PyGenUtil::TCHARToUTF8Buffer(*TypeDocString);
 
 	GeneratedWrappedType->PyType.tp_basicsize = sizeof(FPyWrapperObject);
 	GeneratedWrappedType->PyType.tp_base = SuperPyType ? SuperPyType : &PyWrapperObjectType;
 	GeneratedWrappedType->PyType.tp_flags = Py_TPFLAGS_DEFAULT | Py_TPFLAGS_BASETYPE;
 
-	TSharedPtr<FPyWrapperObjectMetaData> ObjectMetaData = MakeShared<FPyWrapperObjectMetaData>();
+	TSharedRef<FPyWrapperObjectMetaData> ObjectMetaData = MakeShared<FPyWrapperObjectMetaData>();
 	ObjectMetaData->Class = (UClass*)InClass;
 	ObjectMetaData->PythonProperties = MoveTemp(PythonProperties);
 	ObjectMetaData->PythonMethods = MoveTemp(PythonMethods);
@@ -764,6 +838,9 @@ PyTypeObject* FPyWrapperTypeRegistry::GetWrappedClassType(const UClass* InClass)
 
 PyTypeObject* FPyWrapperTypeRegistry::GenerateWrappedStructType(const UStruct* InStruct, FGeneratedWrappedTypeReferences& OutGeneratedWrappedTypeReferences, TSet<FName>& OutDirtyModules, const bool bForceGenerate)
 {
+	SCOPE_SECONDS_ACCUMULATOR(STAT_GenerateWrappedStructTotalTime);
+	INC_DWORD_STAT(STAT_GenerateWrappedStructCallCount);
+
 	struct FFuncs
 	{
 		static int Init(FPyWrapperStruct* InSelf, PyObject* InArgs, PyObject* InKwds)
@@ -809,8 +886,11 @@ PyTypeObject* FPyWrapperTypeRegistry::GenerateWrappedStructType(const UStruct* I
 		SuperPyType = GenerateWrappedStructType(SuperStruct, OutGeneratedWrappedTypeReferences, OutDirtyModules, true);
 	}
 
+	INC_DWORD_STAT(STAT_GenerateWrappedStructObjCount);
+
 	check(!GeneratedWrappedTypes.Contains(InStruct->GetFName()));
-	TSharedPtr<PyGenUtil::FGeneratedWrappedType> GeneratedWrappedType = GeneratedWrappedTypes.Add(InStruct->GetFName(), MakeShared<PyGenUtil::FGeneratedWrappedType>());
+	TSharedRef<PyGenUtil::FGeneratedWrappedStructType> GeneratedWrappedType = MakeShared<PyGenUtil::FGeneratedWrappedStructType>();
+	GeneratedWrappedTypes.Add(InStruct->GetFName(), GeneratedWrappedType);
 
 	TMap<FName, FName> PythonProperties;
 
@@ -823,15 +903,15 @@ PyTypeObject* FPyWrapperTypeRegistry::GenerateWrappedStructType(const UStruct* I
 		{
 			GatherWrappedTypesForPropertyReferences(InProp, OutGeneratedWrappedTypeReferences);
 
-			const PyGenUtil::FGeneratedWrappedPropertyDoc& GeneratedPropertyDoc = GeneratedWrappedType->TypePropertyDocs[GeneratedWrappedType->TypePropertyDocs.Emplace(InProp)];
+			const PyGenUtil::FGeneratedWrappedPropertyDoc& GeneratedPropertyDoc = GeneratedWrappedType->PropertyDocs[GeneratedWrappedType->PropertyDocs.Emplace(InProp)];
 			PythonProperties.Add(*GeneratedPropertyDoc.PythonPropName, InProp->GetFName());
 
 			if (bExportPropertyToScript)
 			{
-				PyGenUtil::FGeneratedWrappedGetSet& GeneratedWrappedGetSet = GeneratedWrappedType->TypeGetSets[GeneratedWrappedType->TypeGetSets.AddDefaulted()];
+				PyGenUtil::FGeneratedWrappedGetSet& GeneratedWrappedGetSet = GeneratedWrappedType->GetSets.TypeGetSets.AddDefaulted_GetRef();
 				GeneratedWrappedGetSet.GetSetName = PyGenUtil::TCHARToUTF8Buffer(*GeneratedPropertyDoc.PythonPropName);
 				GeneratedWrappedGetSet.GetSetDoc = PyGenUtil::TCHARToUTF8Buffer(*GeneratedPropertyDoc.DocString);
-				GeneratedWrappedGetSet.PropName = InProp->GetFName();
+				GeneratedWrappedGetSet.Prop = InProp;
 				GeneratedWrappedGetSet.GetCallback = (getter)&FPyWrapperStruct::Getter_Impl;
 				GeneratedWrappedGetSet.SetCallback = (setter)&FPyWrapperStruct::Setter_Impl;
 			}
@@ -849,14 +929,14 @@ PyTypeObject* FPyWrapperTypeRegistry::GenerateWrappedStructType(const UStruct* I
 	FString TypeDocString = PyGenUtil::PythonizeTooltip(PyGenUtil::GetFieldTooltip(InStruct));
 	if (const UStruct* SuperStruct = InStruct->GetSuperStruct())
 	{
-		TSharedPtr<PyGenUtil::FGeneratedWrappedType> SuperGeneratedWrappedType = GeneratedWrappedTypes.FindRef(SuperStruct->GetFName());
+		TSharedPtr<PyGenUtil::FGeneratedWrappedStructType> SuperGeneratedWrappedType = StaticCastSharedPtr<PyGenUtil::FGeneratedWrappedStructType>(GeneratedWrappedTypes.FindRef(SuperStruct->GetFName()));
 		if (SuperGeneratedWrappedType.IsValid())
 		{
-			GeneratedWrappedType->TypePropertyDocs.Append(SuperGeneratedWrappedType->TypePropertyDocs);
+			GeneratedWrappedType->PropertyDocs.Append(SuperGeneratedWrappedType->PropertyDocs);
 		}
 	}
-	GeneratedWrappedType->TypePropertyDocs.Sort(&PyGenUtil::FGeneratedWrappedPropertyDoc::SortPredicate);
-	PyGenUtil::FGeneratedWrappedPropertyDoc::AppendDocString(GeneratedWrappedType->TypePropertyDocs, TypeDocString, /*bEditorVariant*/true);
+	GeneratedWrappedType->PropertyDocs.Sort(&PyGenUtil::FGeneratedWrappedPropertyDoc::SortPredicate);
+	PyGenUtil::FGeneratedWrappedPropertyDoc::AppendDocString(GeneratedWrappedType->PropertyDocs, TypeDocString, /*bEditorVariant*/true);
 	GeneratedWrappedType->TypeDoc = PyGenUtil::TCHARToUTF8Buffer(*TypeDocString);
 
 	GeneratedWrappedType->PyType.tp_basicsize = sizeof(FPyWrapperStruct);
@@ -864,7 +944,7 @@ PyTypeObject* FPyWrapperTypeRegistry::GenerateWrappedStructType(const UStruct* I
 	GeneratedWrappedType->PyType.tp_init = (initproc)&FFuncs::Init;
 	GeneratedWrappedType->PyType.tp_flags = Py_TPFLAGS_DEFAULT | Py_TPFLAGS_BASETYPE;
 
-	TSharedPtr<FPyWrapperStructMetaData> StructMetaData = MakeShared<FPyWrapperStructMetaData>();
+	TSharedRef<FPyWrapperStructMetaData> StructMetaData = MakeShared<FPyWrapperStructMetaData>();
 	StructMetaData->Struct = (UStruct*)InStruct;
 	StructMetaData->PythonProperties = MoveTemp(PythonProperties);
 	// Build a complete list of init params for this struct (parent struct params + our params)
@@ -876,11 +956,11 @@ PyTypeObject* FPyWrapperTypeRegistry::GenerateWrappedStructType(const UStruct* I
 			StructMetaData->InitParams = SuperMetaData->InitParams;
 		}
 	}
-	for (const PyGenUtil::FGeneratedWrappedGetSet& GeneratedWrappedGetSet : GeneratedWrappedType->TypeGetSets)
+	for (const PyGenUtil::FGeneratedWrappedGetSet& GeneratedWrappedGetSet : GeneratedWrappedType->GetSets.TypeGetSets)
 	{
-		PyGenUtil::FGeneratedWrappedMethodParameter& GeneratedInitParam = StructMetaData->InitParams[StructMetaData->InitParams.AddDefaulted()];
+		PyGenUtil::FGeneratedWrappedMethodParameter& GeneratedInitParam = StructMetaData->InitParams.AddDefaulted_GetRef();
 		GeneratedInitParam.ParamName = GeneratedWrappedGetSet.GetSetName;
-		GeneratedInitParam.ParamPropName = GeneratedWrappedGetSet.PropName;
+		GeneratedInitParam.ParamProp = GeneratedWrappedGetSet.Prop;
 		GeneratedInitParam.ParamDefaultValue = FString();
 	}
 	GeneratedWrappedType->MetaData = StructMetaData;
@@ -928,6 +1008,9 @@ PyTypeObject* FPyWrapperTypeRegistry::GetWrappedStructType(const UStruct* InStru
 
 PyTypeObject* FPyWrapperTypeRegistry::GenerateWrappedEnumType(const UEnum* InEnum, FGeneratedWrappedTypeReferences& OutGeneratedWrappedTypeReferences, TSet<FName>& OutDirtyModules, const bool bForceGenerate)
 {
+	SCOPE_SECONDS_ACCUMULATOR(STAT_GenerateWrappedEnumTotalTime);
+	INC_DWORD_STAT(STAT_GenerateWrappedEnumCallCount);
+
 	// Already processed? Nothing more to do
 	if (PyTypeObject* ExistingPyType = PythonWrappedEnums.FindRef(InEnum->GetFName()))
 	{
@@ -946,8 +1029,11 @@ PyTypeObject* FPyWrapperTypeRegistry::GenerateWrappedEnumType(const UEnum* InEnu
 		return nullptr;
 	}
 
+	INC_DWORD_STAT(STAT_GenerateWrappedEnumObjCount);
+
 	check(!GeneratedWrappedTypes.Contains(InEnum->GetFName()));
-	TSharedPtr<PyGenUtil::FGeneratedWrappedType> GeneratedWrappedType = GeneratedWrappedTypes.Add(InEnum->GetFName(), MakeShared<PyGenUtil::FGeneratedWrappedType>());
+	TSharedRef<PyGenUtil::FGeneratedWrappedType> GeneratedWrappedType = MakeShared<PyGenUtil::FGeneratedWrappedType>();
+	GeneratedWrappedTypes.Add(InEnum->GetFName(), GeneratedWrappedType);
 
 	GeneratedWrappedType->TypeName = PyGenUtil::TCHARToUTF8Buffer(*PyGenUtil::GetEnumPythonName(InEnum));
 	GeneratedWrappedType->TypeDoc = PyGenUtil::TCHARToUTF8Buffer(*PyGenUtil::PythonizeTooltip(PyGenUtil::GetFieldTooltip(InEnum)));
@@ -956,7 +1042,7 @@ PyTypeObject* FPyWrapperTypeRegistry::GenerateWrappedEnumType(const UEnum* InEnu
 	GeneratedWrappedType->PyType.tp_base = &PyWrapperEnumType;
 	GeneratedWrappedType->PyType.tp_flags = Py_TPFLAGS_DEFAULT;
 
-	TSharedPtr<FPyWrapperEnumMetaData> EnumMetaData = MakeShared<FPyWrapperEnumMetaData>();
+	TSharedRef<FPyWrapperEnumMetaData> EnumMetaData = MakeShared<FPyWrapperEnumMetaData>();
 	EnumMetaData->Enum = (UEnum*)InEnum;
 	GeneratedWrappedType->MetaData = EnumMetaData;
 
@@ -967,13 +1053,13 @@ PyTypeObject* FPyWrapperTypeRegistry::GenerateWrappedEnumType(const UEnum* InEnu
 		{
 			if (PyGenUtil::ShouldExportEnumEntry(InEnum, EnumEntryIndex))
 			{
-			const int64 EnumEntryValue = InEnum->GetValueByIndex(EnumEntryIndex);
-			const FString EnumEntryShortName = InEnum->GetNameStringByIndex(EnumEntryIndex);
-			const FString EnumEntryShortPythonName = PyGenUtil::PythonizeName(EnumEntryShortName, PyGenUtil::EPythonizeNameCase::Upper);
-			const FString EnumEntryDoc = PyGenUtil::PythonizeTooltip(PyGenUtil::GetEnumEntryTooltip(InEnum, EnumEntryIndex));
+				const int64 EnumEntryValue = InEnum->GetValueByIndex(EnumEntryIndex);
+				const FString EnumEntryShortName = InEnum->GetNameStringByIndex(EnumEntryIndex);
+				const FString EnumEntryShortPythonName = PyGenUtil::PythonizeName(EnumEntryShortName, PyGenUtil::EPythonizeNameCase::Upper);
+				const FString EnumEntryDoc = PyGenUtil::PythonizeTooltip(PyGenUtil::GetEnumEntryTooltip(InEnum, EnumEntryIndex));
 
-			FPyWrapperEnum::SetEnumEntryValue(&GeneratedWrappedType->PyType, EnumEntryValue, TCHAR_TO_UTF8(*EnumEntryShortPythonName), TCHAR_TO_UTF8(*EnumEntryDoc));
-		}
+				FPyWrapperEnum::SetEnumEntryValue(&GeneratedWrappedType->PyType, EnumEntryValue, TCHAR_TO_UTF8(*EnumEntryShortPythonName), TCHAR_TO_UTF8(*EnumEntryDoc));
+			}
 		}
 
 		const FName UnrealModuleName = *PyGenUtil::GetFieldModule(InEnum);
@@ -1013,6 +1099,9 @@ PyTypeObject* FPyWrapperTypeRegistry::GetWrappedEnumType(const UEnum* InEnum) co
 
 PyTypeObject* FPyWrapperTypeRegistry::GenerateWrappedDelegateType(const UFunction* InDelegateSignature, FGeneratedWrappedTypeReferences& OutGeneratedWrappedTypeReferences, TSet<FName>& OutDirtyModules, const bool bForceGenerate)
 {
+	SCOPE_SECONDS_ACCUMULATOR(STAT_GenerateWrappedDelegateTotalTime);
+	INC_DWORD_STAT(STAT_GenerateWrappedDelegateCallCount);
+
 	// Already processed? Nothing more to do
 	if (PyTypeObject* ExistingPyType = PythonWrappedDelegates.FindRef(InDelegateSignature->GetFName()))
 	{
@@ -1025,8 +1114,11 @@ PyTypeObject* FPyWrapperTypeRegistry::GenerateWrappedDelegateType(const UFunctio
 		return nullptr;
 	}
 
+	INC_DWORD_STAT(STAT_GenerateWrappedDelegateObjCount);
+
 	check(!GeneratedWrappedTypes.Contains(InDelegateSignature->GetFName()));
-	TSharedPtr<PyGenUtil::FGeneratedWrappedType> GeneratedWrappedType = GeneratedWrappedTypes.Add(InDelegateSignature->GetFName(), MakeShared<PyGenUtil::FGeneratedWrappedType>());
+	TSharedRef<PyGenUtil::FGeneratedWrappedType> GeneratedWrappedType = MakeShared<PyGenUtil::FGeneratedWrappedType>();
+	GeneratedWrappedTypes.Add(InDelegateSignature->GetFName(), GeneratedWrappedType);
 
 	for (TFieldIterator<const UProperty> ParamIt(InDelegateSignature); ParamIt; ++ParamIt)
 	{
@@ -1045,8 +1137,8 @@ PyTypeObject* FPyWrapperTypeRegistry::GenerateWrappedDelegateType(const UFunctio
 		GeneratedWrappedType->PyType.tp_basicsize = sizeof(FPyWrapperMulticastDelegate);
 		GeneratedWrappedType->PyType.tp_base = &PyWrapperMulticastDelegateType;
 
-		TSharedPtr<FPyWrapperMulticastDelegateMetaData> DelegateMetaData = MakeShared<FPyWrapperMulticastDelegateMetaData>();
-		DelegateMetaData->DelegateSignature = InDelegateSignature;
+		TSharedRef<FPyWrapperMulticastDelegateMetaData> DelegateMetaData = MakeShared<FPyWrapperMulticastDelegateMetaData>();
+		DelegateMetaData->DelegateSignature.SetFunctionAndExtractParams(InDelegateSignature);
 		GeneratedWrappedType->MetaData = DelegateMetaData;
 	}
 	else
@@ -1054,8 +1146,8 @@ PyTypeObject* FPyWrapperTypeRegistry::GenerateWrappedDelegateType(const UFunctio
 		GeneratedWrappedType->PyType.tp_basicsize = sizeof(FPyWrapperDelegate);
 		GeneratedWrappedType->PyType.tp_base = &PyWrapperDelegateType;
 
-		TSharedPtr<FPyWrapperDelegateMetaData> DelegateMetaData = MakeShared<FPyWrapperDelegateMetaData>();
-		DelegateMetaData->DelegateSignature = InDelegateSignature;
+		TSharedRef<FPyWrapperDelegateMetaData> DelegateMetaData = MakeShared<FPyWrapperDelegateMetaData>();
+		DelegateMetaData->DelegateSignature.SetFunctionAndExtractParams(InDelegateSignature);
 		GeneratedWrappedType->MetaData = DelegateMetaData;
 	}
 
diff --git a/Engine/Plugins/Experimental/PythonScriptPlugin/Source/PythonScriptPlugin/Private/PythonScriptCommandlet.cpp b/Engine/Plugins/Experimental/PythonScriptPlugin/Source/PythonScriptPlugin/Private/PythonScriptCommandlet.cpp
new file mode 100644
--- /dev/null
+++ b/Engine/Plugins/Experimental/PythonScriptPlugin/Source/PythonScriptPlugin/Private/PythonScriptCommandlet.cpp
@@ -0,0 +1,36 @@
+// Copyright 1998-2018 Epic Games, Inc. All Rights Reserved.
+
+#include "PythonScriptCommandlet.h"
+#include "PythonScriptPlugin.h"
+
+DEFINE_LOG_CATEGORY_STATIC(LogPythonScriptCommandlet, Log, All);
+
+int32 UPythonScriptCommandlet::Main(const FString& Params)
+{
+	TArray<FString> Tokens;
+	TArray<FString> Switches;
+	TMap<FString, FString> ParamVals;
+	UCommandlet::ParseCommandLine(*Params, Tokens, Switches, ParamVals);
+
+	// Find the Python script to run
+	FString PythonScript;
+	if (const FString* ParamVal = ParamVals.Find(TEXT("Script")))
+	{
+		PythonScript = *ParamVal;
+	}
+	else
+	{
+		UE_LOG(LogPythonScriptCommandlet, Error, TEXT("-Script argument not specified"));
+		return -1;
+	}
+
+#if WITH_PYTHON
+	UE_LOG(LogPythonScriptCommandlet, Display, TEXT("Running Python script: %s"), *PythonScript);
+	FPythonScriptPlugin::Get()->HandlePythonExecCommand(*PythonScript);
+#else	// WITH_PYTHON
+	UE_LOG(LogPythonScriptCommandlet, Error, TEXT("Python script cannot run as the plugin was built as a stub!"));
+	return -1;
+#endif	// WITH_PYTHON
+
+	return 0;
+}
diff --git a/Engine/Plugins/Experimental/PythonScriptPlugin/Source/PythonScriptPlugin/Private/PythonScriptPlugin.cpp b/Engine/Plugins/Experimental/PythonScriptPlugin/Source/PythonScriptPlugin/Private/PythonScriptPlugin.cpp
--- a/Engine/Plugins/Experimental/PythonScriptPlugin/Source/PythonScriptPlugin/Private/PythonScriptPlugin.cpp
+++ b/Engine/Plugins/Experimental/PythonScriptPlugin/Source/PythonScriptPlugin/Private/PythonScriptPlugin.cpp
@@ -1,7 +1,9 @@
 // Copyright 1998-2018 Epic Games, Inc. All Rights Reserved.
 
 #include "PythonScriptPlugin.h"
+#include "PythonScriptPluginSettings.h"
 #include "PyCore.h"
+#include "PySlate.h"
 #include "PyEditor.h"
 #include "PyConstant.h"
 #include "PyConversion.h"
@@ -19,7 +21,6 @@
 #include "HAL/PlatformProcess.h"
 #include "Containers/Ticker.h"
 #include "Features/IModularFeatures.h"
-#include "Modules/ModuleManager.h"
 #include "ProfilingDebugging/ScopedTimers.h"
 
 #if WITH_EDITOR
@@ -38,6 +39,41 @@
 
 #if WITH_PYTHON
 
+static PyUtil::FPyApiBuffer NullPyArg = PyUtil::TCHARToPyApiBuffer(TEXT(""));
+static PyUtil::FPyApiChar* NullPyArgPtrs[] = { NullPyArg.GetData() };
+
+/** Util struct to set the sys.argv data for Python when executing a file with arguments */
+struct FPythonScopedArgv
+{
+	FPythonScopedArgv(const TCHAR* InArgs)
+	{
+		if (InArgs && *InArgs)
+		{
+			FString NextToken;
+			while (FParse::Token(InArgs, NextToken, false))
+			{
+				PyCommandLineArgs.Add(PyUtil::TCHARToPyApiBuffer(*NextToken));
+			}
+
+			PyCommandLineArgPtrs.Reserve(PyCommandLineArgs.Num());
+			for (PyUtil::FPyApiBuffer& PyCommandLineArg : PyCommandLineArgs)
+			{
+				PyCommandLineArgPtrs.Add(PyCommandLineArg.GetData());
+			}
+
+			PySys_SetArgvEx(PyCommandLineArgPtrs.Num(), PyCommandLineArgPtrs.GetData(), 0);
+		}
+	}
+
+	~FPythonScopedArgv()
+	{
+		PySys_SetArgvEx(1, NullPyArgPtrs, 0);
+	}
+
+	TArray<PyUtil::FPyApiBuffer> PyCommandLineArgs;
+	TArray<PyUtil::FPyApiChar*> PyCommandLineArgPtrs;
+};
+
 FPythonCommandExecutor::FPythonCommandExecutor(FPythonScriptPlugin* InPythonScriptPlugin)
 	: PythonScriptPlugin(InPythonScriptPlugin)
 {
@@ -305,14 +341,24 @@ FPythonScriptPlugin::FPythonScriptPlugin()
 	: CmdExec(this)
 	, CmdMenu(nullptr)
 	, bInitialized(false)
+	, bHasTicked(false)
 #endif	// WITH_PYTHON
 {
 }
 
-FPythonScriptPlugin* FPythonScriptPlugin::Get()
+bool FPythonScriptPlugin::IsPythonAvailable() const
 {
-	static const FName ModuleName = "PythonScriptPlugin";
-	return FModuleManager::GetModulePtr<FPythonScriptPlugin>(ModuleName);
+	return WITH_PYTHON;
+}
+
+bool FPythonScriptPlugin::ExecPythonCommand(const TCHAR* InPythonCommand)
+{
+#if WITH_PYTHON
+	return HandlePythonExecCommand(InPythonCommand);
+#else	// WITH_PYTHON
+	ensureAlwaysMsgf(false, TEXT("Python is not available!"));
+	return false;
+#endif	// WITH_PYTHON
 }
 
 void FPythonScriptPlugin::StartupModule()
@@ -383,22 +429,13 @@ void FPythonScriptPlugin::InitializePython()
 		PyHomePath = PyUtil::TCHARToPyApiBuffer(*PythonDir);
 	}
 
-	// Set-up the correct command line
 	{
-		PyCommandLineArgs.Add(PyUtil::TCHARToPyApiBuffer(TEXT(""))); // Script name; always empty
-
-		const TCHAR* CmdLine = FCommandLine::Get();
-		FString NextToken;
-		while (FParse::Token(CmdLine, NextToken, false))
-		{
-			PyCommandLineArgs.Add(PyUtil::TCHARToPyApiBuffer(*NextToken));
-		}
-
-		PyCommandLineArgPtrs.Reserve(PyCommandLineArgs.Num());
-		for (PyUtil::FPyApiBuffer& PyCommandLineArg : PyCommandLineArgs)
-		{
-			PyCommandLineArgPtrs.Add(PyCommandLineArg.GetData());
-		}
+		// Build the full Python directory (UE_PYTHON_DIR may be relative to UE4 engine directory for portability)
+		FString PythonDir = UTF8_TO_TCHAR(UE_PYTHON_DIR);
+		PythonDir.ReplaceInline(TEXT("{ENGINE_DIR}"), *FPaths::EngineDir(), ESearchCase::CaseSensitive);
+		FPaths::NormalizeDirectoryName(PythonDir);
+		FPaths::RemoveDuplicateSlashes(PythonDir);
+		PyHomePath = PyUtil::TCHARToPyApiBuffer(*PythonDir);
 	}
 
 	// Initialize the Python interpreter
@@ -410,7 +447,13 @@ void FPythonScriptPlugin::InitializePython()
 		Py_SetPythonHome(PyHomePath.GetData());
 		Py_Initialize();
 
-		PySys_SetArgv(PyCommandLineArgPtrs.Num(), PyCommandLineArgPtrs.GetData());
+		PySys_SetArgvEx(1, NullPyArgPtrs, 0);
+
+		// Enable developer warnings if requested
+		if (GetDefault<UPythonScriptPluginSettings>()->bDeveloperMode)
+		{
+			PyUtil::EnableDeveloperWarnings();
+		}
 
 		// Initialize our custom method type as we'll need it when generating bindings
 		InitializePyMethodWithClosure();
@@ -439,6 +482,11 @@ void FPythonScriptPlugin::InitializePython()
 			PyUtil::AddSystemPath(FPaths::ConvertRelativePathToFull(RootFilesystemPath / TEXT("Python")));
 		}
 
+		for (const FDirectoryPath& AdditionalPath : GetDefault<UPythonScriptPluginSettings>()->AdditionalPaths)
+		{
+			PyUtil::AddSystemPath(FPaths::ConvertRelativePathToFull(AdditionalPath.Path));
+		}
+
 		FPackageName::OnContentPathMounted().AddRaw(this, &FPythonScriptPlugin::OnContentPathMounted);
 		FPackageName::OnContentPathDismounted().AddRaw(this, &FPythonScriptPlugin::OnContentPathDismounted);
 	}
@@ -452,6 +500,10 @@ void FPythonScriptPlugin::InitializePython()
 		PyCore::InitializeModule();
 		ImportUnrealModule(TEXT("core"));
 
+		// Initialize the and import the "slate" module
+		PySlate::InitializeModule();
+		ImportUnrealModule(TEXT("slate"));
+
 #if WITH_EDITOR
 		// Initialize the and import the "editor" module
 		PyEditor::InitializeModule();
@@ -464,10 +516,10 @@ void FPythonScriptPlugin::InitializePython()
 		// Initialize the wrapped types
 		FPyWrapperTypeRegistry::Get().GenerateWrappedTypes();
 
-		// Initialize the re-instancer ticker
-		ReinstanceTickerHandle = FTicker::GetCoreTicker().AddTicker(FTickerDelegate::CreateLambda([](float DeltaTime)
+		// Initialize the tick handler
+		TickHandle = FTicker::GetCoreTicker().AddTicker(FTickerDelegate::CreateLambda([this](float DeltaTime)
 		{
-			FPyWrapperTypeReinstancer::Get().ProcessPending();
+			Tick(DeltaTime);
 			return true;
 		}));
 	}
@@ -480,7 +532,7 @@ void FPythonScriptPlugin::ShutdownPython()
 		return;
 	}
 
-	FTicker::GetCoreTicker().RemoveTicker(ReinstanceTickerHandle);
+	FTicker::GetCoreTicker().RemoveTicker(TickHandle);
 
 	FPyWrapperTypeRegistry::Get().OnModuleDirtied().RemoveAll(this);
 	FModuleManager::Get().OnModulesChanged().RemoveAll(this);
@@ -501,6 +553,24 @@ void FPythonScriptPlugin::ShutdownPython()
 	Py_Finalize();
 
 	bInitialized = false;
+	bHasTicked = false;
+}
+
+void FPythonScriptPlugin::Tick(const float InDeltaTime)
+{
+	// If this is our first Tick, handle any post-init logic that should happen once the engine is fully initialized
+	if (!bHasTicked)
+	{
+		// Run start-up scripts now
+		for (const FString& StartupScript : GetDefault<UPythonScriptPluginSettings>()->StartupScripts)
+		{
+			HandlePythonExecCommand(*StartupScript);
+		}
+	}
+
+	FPyWrapperTypeReinstancer::Get().ProcessPending();
+
+	bHasTicked = true;
 }
 
 void FPythonScriptPlugin::ImportUnrealModule(const TCHAR* InModuleName)
@@ -564,15 +634,24 @@ void FPythonScriptPlugin::ImportUnrealModule(const TCHAR* InModuleName)
 	}
 }
 
-void FPythonScriptPlugin::HandlePythonExecCommand(const TCHAR* InPythonCommand)
+bool FPythonScriptPlugin::HandlePythonExecCommand(const TCHAR* InPythonCommand)
 {
-	if (FPaths::GetExtension(InPythonCommand) == TEXT("py"))
+	// We may have been passed literal code or a file
+	// To work out which, extract the first token and see if it's a .py file
+	// If it is, treat the remaining text as arguments to the file
+	// Otherwise, treat it as literal code
+	FString ExtractedFilename;
+	{
+		const TCHAR* Tmp = InPythonCommand;
+		ExtractedFilename = FParse::Token(Tmp, false);
+	}
+	if (FPaths::GetExtension(ExtractedFilename) == TEXT("py"))
 	{
-		RunFile(InPythonCommand);
+		return RunFile(*ExtractedFilename, InPythonCommand);
 	}
 	else
 	{
-		RunString(InPythonCommand);
+		return RunString(InPythonCommand);
 	}
 }
 
@@ -611,18 +690,20 @@ PyObject* FPythonScriptPlugin::EvalString(const TCHAR* InStr, const TCHAR* InCon
 	return PyEval_EvalCode((PyUtil::FPyCodeObjectType*)PyCodeObj.Get(), InGlobalDict, InLocalDict);
 }
 
-void FPythonScriptPlugin::RunString(const TCHAR* InStr)
+bool FPythonScriptPlugin::RunString(const TCHAR* InStr)
 {
 	FPyObjectPtr PyResult = FPyObjectPtr::StealReference(EvalString(InStr, TEXT("<string>"), Py_file_input));
 	if (!PyResult)
 	{
 		PyUtil::LogPythonError();
+		return false;
 	}
 
 	FPyWrapperTypeReinstancer::Get().ProcessPending();
+	return true;
 }
 
-void FPythonScriptPlugin::RunFile(const TCHAR* InFile)
+bool FPythonScriptPlugin::RunFile(const TCHAR* InFile, const TCHAR* InArgs)
 {
 	auto ResolveFilePath = [InFile]() -> FString
 	{
@@ -677,7 +758,7 @@ void FPythonScriptPlugin::RunFile(const TCHAR* InFile)
 	if (!bLoaded)
 	{
 		UE_LOG(LogPython, Error, TEXT("Could not load Python file '%s' (resolved from '%s')"), *ResolvedFilePath, InFile);
-		return;
+		return false;
 	}
 
 	FPyObjectPtr PyFileGlobalDict = FPyObjectPtr::StealReference(PyDict_Copy(PyGlobalDict));
@@ -694,6 +775,7 @@ void FPythonScriptPlugin::RunFile(const TCHAR* InFile)
 	FPyObjectPtr PyResult;
 	{
 		FScopedDurationTimer Timer(ElapsedSeconds);
+		FPythonScopedArgv ScopedArgv(InArgs);
 
 		// We can't just use PyRun_File here as Python isn't always built against the same version of the CRT as UE4, so we get a crash at the CRT layer
 		PyResult = FPyObjectPtr::StealReference(EvalString(*FileStr, *ResolvedFilePath, Py_file_input, PyFileGlobalDict, PyFileLocalDict));
@@ -702,6 +784,7 @@ void FPythonScriptPlugin::RunFile(const TCHAR* InFile)
 	if (!PyResult)
 	{
 		PyUtil::LogPythonError();
+		return false;
 	}
 
 	FPyWrapperTypeReinstancer::Get().ProcessPending();
@@ -712,6 +795,8 @@ void FPythonScriptPlugin::RunFile(const TCHAR* InFile)
 		EventAttributes.Add(FAnalyticsEventAttribute(TEXT("Duration"), ElapsedSeconds));
 		FEngineAnalytics::GetProvider().RecordEvent(TEXT("PythonScriptPlugin"), EventAttributes);
 	}
+
+	return true;
 }
 
 void FPythonScriptPlugin::OnModuleDirtied(FName InModuleName)
diff --git a/Engine/Plugins/Experimental/PythonScriptPlugin/Source/PythonScriptPlugin/Private/PythonScriptPlugin.h b/Engine/Plugins/Experimental/PythonScriptPlugin/Source/PythonScriptPlugin/Private/PythonScriptPlugin.h
--- a/Engine/Plugins/Experimental/PythonScriptPlugin/Source/PythonScriptPlugin/Private/PythonScriptPlugin.h
+++ b/Engine/Plugins/Experimental/PythonScriptPlugin/Source/PythonScriptPlugin/Private/PythonScriptPlugin.h
@@ -3,10 +3,10 @@
 #pragma once
 
 #include "IncludePython.h"
+#include "IPythonScriptPlugin.h"
 #include "PyUtil.h"
 #include "PyPtr.h"
 #include "Misc/CoreMisc.h"
-#include "Modules/ModuleInterface.h"
 #include "HAL/IConsoleManager.h"
 
 class FPythonScriptPlugin;
@@ -50,13 +50,20 @@ struct IPythonCommandMenu
 };
 #endif	// WITH_PYTHON
 
-class FPythonScriptPlugin : public IModuleInterface, public FSelfRegisteringExec
+class FPythonScriptPlugin : public IPythonScriptPlugin, public FSelfRegisteringExec
 {
 public:
 	FPythonScriptPlugin();
 
 	/** Get this module */
-	static FPythonScriptPlugin* Get();
+	static FPythonScriptPlugin* Get()
+	{
+		return static_cast<FPythonScriptPlugin*>(IPythonScriptPlugin::Get());
+	}
+
+	//~ IPythonScriptPlugin interface
+	virtual bool IsPythonAvailable() const override;
+	virtual bool ExecPythonCommand(const TCHAR* InPythonCommand) override;
 
 	//~ IModuleInterface interface
 	virtual void StartupModule() override;
@@ -72,14 +79,14 @@ public:
 	 */
 	void ImportUnrealModule(const TCHAR* InModuleName);
 
-	void HandlePythonExecCommand(const TCHAR* InPythonCommand);
+	bool HandlePythonExecCommand(const TCHAR* InPythonCommand);
 
 	PyObject* EvalString(const TCHAR* InStr, const TCHAR* InContext, const int InMode);
 	PyObject* EvalString(const TCHAR* InStr, const TCHAR* InContext, const int InMode, PyObject* InGlobalDict, PyObject* InLocalDict);
 
-	void RunString(const TCHAR* InStr);
+	bool RunString(const TCHAR* InStr);
 
-	void RunFile(const TCHAR* InFile);
+	bool RunFile(const TCHAR* InFile, const TCHAR* InArgs);
 #endif	// WITH_PYTHON
 
 private:
@@ -88,6 +95,8 @@ private:
 
 	void ShutdownPython();
 
+	void Tick(const float InDeltaTime);
+
 	void OnModuleDirtied(FName InModuleName);
 
 	void OnModulesChanged(FName InModuleName, EModuleChangeReason InModuleChangeReason);
@@ -102,16 +111,15 @@ private:
 
 	FPythonCommandExecutor CmdExec;
 	IPythonCommandMenu* CmdMenu;
-	FDelegateHandle ReinstanceTickerHandle;
+	FDelegateHandle TickHandle;
 
 	PyUtil::FPyApiBuffer PyProgramName;
 	PyUtil::FPyApiBuffer PyHomePath;
-	TArray<PyUtil::FPyApiBuffer> PyCommandLineArgs;
-	TArray<PyUtil::FPyApiChar*> PyCommandLineArgPtrs;
 	FPyObjectPtr PyGlobalDict;
 	FPyObjectPtr PyLocalDict;
 	FPyObjectPtr PyUnrealModule;
 	bool bInitialized;
+	bool bHasTicked;
 #endif	// WITH_PYTHON
 
 
diff --git a/Engine/Plugins/Online/OnlineSubsystemUtils/Source/OnlineSubsystemUtils/Private/IpNetDriver.cpp b/Engine/Plugins/Online/OnlineSubsystemUtils/Source/OnlineSubsystemUtils/Private/IpNetDriver.cpp
--- a/Engine/Plugins/Online/OnlineSubsystemUtils/Source/OnlineSubsystemUtils/Private/IpNetDriver.cpp
+++ b/Engine/Plugins/Online/OnlineSubsystemUtils/Source/OnlineSubsystemUtils/Private/IpNetDriver.cpp
@@ -441,7 +441,7 @@ void UIpNetDriver::TickDispatch( float DeltaTime )
 						}
 #endif
 
-						Connection->InitRemoteConnection( this, Socket,  FURL(), *FromAddr, USOCK_Open);
+						Connection->InitRemoteConnection( this, Socket, World ? World->URL : FURL(), *FromAddr, USOCK_Open);
 
 						if (Connection->Handler.IsValid())
 						{
diff --git a/Engine/Plugins/Runtime/AppleARKit/Source/AppleARKit/Private/AppleARKitLiveLinkSource.cpp b/Engine/Plugins/Runtime/AppleARKit/Source/AppleARKit/Private/AppleARKitLiveLinkSource.cpp
--- a/Engine/Plugins/Runtime/AppleARKit/Source/AppleARKit/Private/AppleARKitLiveLinkSource.cpp
+++ b/Engine/Plugins/Runtime/AppleARKit/Source/AppleARKit/Private/AppleARKitLiveLinkSource.cpp
@@ -1,4 +1,4 @@
-﻿// Copyright 1998-2017 Epic Games, Inc. All Rights Reserved.
+// Copyright 1998-2017 Epic Games, Inc. All Rights Reserved.
 
 #include "AppleARKitLiveLinkSource.h"
 #include "UObject/Package.h"
@@ -128,7 +128,7 @@ void FAppleARKitLiveLinkSource::PublishBlendShapes(FName SubjectName, double Tim
 		LiveLinkFrame.WorldTime = Timestamp;
 
 		TArray<FLiveLinkCurveElement>& BlendShapes = LiveLinkFrame.CurveElements;
-		
+
 		BlendShapes.Reset((int32)EARFaceBlendShape::MAX);
 
 		// Iterate through all of the blend shapes copying them into the LiveLink data type
diff --git a/Engine/Plugins/Tests/EditorTests/Source/EditorTests/Private/UnrealEd/EditorAssetAutomationTests.cpp b/Engine/Plugins/Tests/EditorTests/Source/EditorTests/Private/UnrealEd/EditorAssetAutomationTests.cpp
--- a/Engine/Plugins/Tests/EditorTests/Source/EditorTests/Private/UnrealEd/EditorAssetAutomationTests.cpp
+++ b/Engine/Plugins/Tests/EditorTests/Source/EditorTests/Private/UnrealEd/EditorAssetAutomationTests.cpp
@@ -6,6 +6,7 @@
 #include "Misc/ConfigCacheIni.h"
 #include "Misc/AutomationTest.h"
 #include "Modules/ModuleManager.h"
+#include "UObject/GCObjectScopeGuard.h"
 #include "UObject/Object.h"
 #include "UObject/Package.h"
 #include "Misc/PackageName.h"
@@ -109,6 +110,7 @@
 #include "Animation/AnimComposite.h"
 #include "Factories/DataAssetFactory.h"
 #include "Factories/CurveFactory.h"
+#include "AssetExportTask.h"
 
 #define LOCTEXT_NAMESPACE "EditorAssetAutomationTests"
 
@@ -1002,23 +1004,23 @@ namespace ImportExportAssetHelper
 
 				//Export the asset
 				const FString ExportAssetName = FString::Printf(TEXT("%s.%s"), *ImportedAsset->GetName(), *Extension);
-				FString ExportPath = FPaths::Combine(*FPaths::AutomationDir(), TEXT("AssetImportExport"), TEXT("Exported"),*ExportAssetName);
-				UExporter* ExporterToUse = UExporter::FindExporter(ImportedAsset, *Extension);
-				
-				UExporter::FExportToFileParams Params;
-				Params.Object = ImportedAsset;
-				Params.Exporter = ExporterToUse;
-				Params.Filename = *ExportPath;
-				Params.InSelectedOnly = false;
-				Params.NoReplaceIdentical = false;
-				Params.Prompt = false;
-				Params.bUseFileArchive = ImportedAsset->IsA(UPackage::StaticClass());
-				Params.WriteEmptyFiles = false;
-				if (UExporter::ExportToFileEx(Params) != 0)  //1 - success, 0 - fatal error, -1 - non fatal error
+
+				UAssetExportTask* ExportTask = NewObject<UAssetExportTask>();
+				FGCObjectScopeGuard ExportTaskGuard(ExportTask);
+				ExportTask->Object = ImportedAsset;
+				ExportTask->Exporter = UExporter::FindExporter(ImportedAsset, *Extension);
+				ExportTask->Filename = FPaths::Combine(*FPaths::AutomationDir(), TEXT("AssetImportExport"), TEXT("Exported"), *ExportAssetName);
+				ExportTask->bSelected = false;
+				ExportTask->bReplaceIdentical = true;
+				ExportTask->bPrompt = false;
+				ExportTask->bUseFileArchive = ImportedAsset->IsA(UPackage::StaticClass());
+				ExportTask->bWriteEmptyFiles = false;
+				ExportTask->bAutomated = true;
+				if (UExporter::RunAssetExportTask(ExportTask) != 0)  //1 - success, 0 - fatal error, -1 - non fatal error
 				{
 					//Success
 					TestReport.bExportSuccessful = true;
-					TestReport.FileSize = IFileManager::Get().FileSize(*ExportPath);
+					TestReport.FileSize = IFileManager::Get().FileSize(*ExportTask->Filename);
 				}
 				else
 				{
diff --git a/Engine/Shaders/Private/Fxaa3_11.ush b/Engine/Shaders/Private/Fxaa3_11.ush
--- a/Engine/Shaders/Private/Fxaa3_11.ush
+++ b/Engine/Shaders/Private/Fxaa3_11.ush
@@ -703,7 +703,20 @@ NOTE the other tuning knobs are now in the shader function inputs!
                    GREEN AS LUMA OPTION SUPPORT FUNCTION
 ============================================================================*/
 #if (FXAA_GREEN_AS_LUMA == 0)
-    FxaaFloat FxaaLuma(FxaaFloat4 rgba) { return rgba.w; }
+
+	#if POST_PROCESS_ALPHA == 2
+		// POST_PROCESS_ALPHA == 2 allows the tonemapper to not output Luma in alpha channel,
+		// therefore need to compute it manually.
+		FxaaFloat FxaaLuma(FxaaFloat4 rgba)
+		{
+			return dot(rgba.rgb, FxaaFloat3(0.299f, 0.587f, 0.114f));
+		}
+
+	#else
+		FxaaFloat FxaaLuma(FxaaFloat4 rgba) { return rgba.w; }
+
+	#endif
+
 #else
     FxaaFloat FxaaLuma(FxaaFloat4 rgba) { return rgba.y; }
 #endif    
diff --git a/Engine/Shaders/Private/PostProcessTemporalAA.usf b/Engine/Shaders/Private/PostProcessTemporalAA.usf
--- a/Engine/Shaders/Private/PostProcessTemporalAA.usf
+++ b/Engine/Shaders/Private/PostProcessTemporalAA.usf
@@ -37,7 +37,11 @@ float OneOverHistoryPreExposure;
 	#endif
 
 #elif TAA_PASS_CONFIG == 1 // Main
-	#define AA_ALPHA (POST_PROCESS_ALPHA)
+	#if POST_PROCESS_ALPHA
+		#define AA_ALPHA 1
+	#else
+		#define AA_ALPHA 0
+	#endif
 	#define AA_BICUBIC 1
 	#define AA_BORDER 1
 	#define AA_CROSS 2
@@ -79,7 +83,11 @@ float OneOverHistoryPreExposure;
 	#define AA_BORDER 1
 
 #elif TAA_PASS_CONFIG == 4 // MainUpsampling
-	#define AA_ALPHA (POST_PROCESS_ALPHA)
+	#if POST_PROCESS_ALPHA
+		#define AA_ALPHA 1
+	#else
+		#define AA_ALPHA 0
+	#endif
 	#define AA_BICUBIC 1
 	#define AA_BORDER 2
 	#define AA_CROSS 1
diff --git a/Engine/Shaders/Private/PostProcessTonemap.usf b/Engine/Shaders/Private/PostProcessTonemap.usf
--- a/Engine/Shaders/Private/PostProcessTonemap.usf
+++ b/Engine/Shaders/Private/PostProcessTonemap.usf
@@ -289,13 +289,13 @@ half3 LookupSceneColor(float2 SceneUV, int2 PixelOffset)
 #endif
 }
 
-float3 SampleSceneColor(float2 SceneUV)
+float4 SampleSceneColor(float2 SceneUV)
 {
 	#if USE_GRAIN_JITTER
 		SceneUV = clamp(SceneUV, SceneUVMinMax.xy, SceneUVMinMax.zw);
 	#endif
 
-	return Texture2DSample(PostprocessInput0, PostprocessInput0Sampler, SceneUV).rgb;
+	return Texture2DSample(PostprocessInput0, PostprocessInput0Sampler, SceneUV);
 }
 
 float4 TonemapCommonPS(
@@ -361,16 +361,16 @@ float4 TonemapCommonPS(
 	        CAUV.xy = ScreenPosToUV(CAUV.xy);
 	        CAUV.zw = ScreenPosToUV(CAUV.zw);
 
-		half3 SceneColor = SampleSceneColor(CAUV.xy + SceneUVJitter.xy).rgb;
+		half4 SceneColor = SampleSceneColor(CAUV.xy + SceneUVJitter.xy);
 		half SceneColorG = SampleSceneColor(CAUV.zw + SceneUVJitter.xy).g;
 		half SceneColorB = SampleSceneColor(SceneUV).b;
 		SceneColor.g = SceneColorG;
 		SceneColor.b = SceneColorB;
 	#else
-		half3 SceneColor = SampleSceneColor(SceneUV).rgb;
+		half4 SceneColor = SampleSceneColor(SceneUV);
 	#endif
 
-	SceneColor *= OneOverPreExposure;
+	SceneColor.rgb *= OneOverPreExposure;
 
 	float ExposureScale = ExposureScaleVignette.x;
 
@@ -407,7 +407,7 @@ float4 TonemapCommonPS(
 //		half A0 = max3(SceneColor);
 //		half4 LuminanceNeightbors = half4(max3(C1), max3(C2), max3(C3), max3(C4));
 		// 5 instructiopns less than reference
-		half A0 = Luminance(SceneColor);
+		half A0 = Luminance(SceneColor.rgb);
 		half4 LuminanceNeightbors = half4(Luminance(C1), Luminance(C2), Luminance(C3), Luminance(C4));
 		// 9 instructions less than reference (green is often a good approximation for the luminance)
 		//half A0 = SceneColor.g;
@@ -439,8 +439,8 @@ float4 TonemapCommonPS(
 //		half3 BlurredSceneColor = ((C1 + C2 + C3 + C4) + SceneColor * 2) / 6;
 //		SceneColor = lerp(SceneColor, BlurredSceneColor, LerpFactor);
 		// optimized, Div6 went into the C++ code
-		half3 DeltaColor = (C1 + C2 + C3 + C4) - SceneColor * 4;
-		SceneColor += DeltaColor * LerpFactor;
+		half3 DeltaColor = (C1 + C2 + C3 + C4) - SceneColor.rgb * 4;
+		SceneColor.rgb += DeltaColor * LerpFactor;
 	}
 #endif
 
@@ -449,11 +449,11 @@ float4 TonemapCommonPS(
 		#if FEATURE_LEVEL > FEATURE_LEVEL_ES3_1 // On mobile ExposureScale applied in base pass PS
 			SceneColor.rgb *= ExposureScale;
 		#endif
-		OutColor.rgb = pow(SceneColor, InverseGamma.x);
+		OutColor.rgb = pow(SceneColor.rgb, InverseGamma.x);
 
 	#else 
 
-		half3 LinearColor = SceneColor * ColorScale0.rgb;
+		half3 LinearColor = SceneColor.rgb * ColorScale0.rgb;
 
 		#if USE_BLOOM
 			float2 BloomUV = clamp(UV, SceneBloomUVMinMax.xy, SceneBloomUVMinMax.zw);
@@ -537,6 +537,10 @@ float4 TonemapCommonPS(
 
 
 	#endif
+	
+	#if POST_PROCESS_ALPHA == 2
+		OutColor.a = SceneColor.a;
+	#endif
 
 #if DEBUG_TONEMAPPER
 	if (!bOutputDebugTile)
diff --git a/Engine/Source/Developer/AssetTools/Private/AssetRenameManager.cpp b/Engine/Source/Developer/AssetTools/Private/AssetRenameManager.cpp
--- a/Engine/Source/Developer/AssetTools/Private/AssetRenameManager.cpp
+++ b/Engine/Source/Developer/AssetTools/Private/AssetRenameManager.cpp
@@ -28,9 +28,9 @@
 #include "Widgets/Views/STableRow.h"
 #include "Widgets/Views/SListView.h"
 #include "EditorStyleSet.h"
-#include "ISourceControlOperation.h"
 #include "SourceControlOperations.h"
 #include "ISourceControlModule.h"
+#include "SourceControlHelpers.h"
 #include "FileHelpers.h"
 #include "SDiscoveringAssetsDialog.h"
 #include "AssetRegistryModule.h"
