Updated wrapped enums to be more consistent with native Python enums

 - Wrapped enums now generate values that are instances of the enum type itself, containing a name and value field (like native Python enums).
 - Wrapped enums are now strongly typed and do not allow implicit conversion from numbers (explicit casting is available, but throws if the value is unknown).
 - Wrapped enum entries may be compared against numbers (even numbers that don't have valid values) via the == and != operators (like IntEnum in Python).
 - Wrapped enums may now be iterated (like native Python enums).
 - Wrapped enums now return a length based on their number of entries (like native Python enums).
 - ScriptName meta-data can now be used with enum entries.

diff --git a/Engine/Plugins/Experimental/PythonScriptPlugin/Source/PythonScriptPlugin/Private/PyConversion.cpp b/Engine/Plugins/Experimental/PythonScriptPlugin/Source/PythonScriptPlugin/Private/PyConversion.cpp
--- a/Engine/Plugins/Experimental/PythonScriptPlugin/Source/PythonScriptPlugin/Private/PyConversion.cpp
+++ b/Engine/Plugins/Experimental/PythonScriptPlugin/Source/PythonScriptPlugin/Private/PyConversion.cpp
@@ -6,6 +6,7 @@
 
 #include "PyWrapperObject.h"
 #include "PyWrapperStruct.h"
+#include "PyWrapperEnum.h"
 #include "PyWrapperDelegate.h"
 #include "PyWrapperName.h"
 #include "PyWrapperText.h"
@@ -57,7 +58,7 @@ FPyConversionResult NativizeStructInstance(PyObject* PyObj, UScriptStruct* Struc
 		StructType->CopyScriptStruct(StructInstance, PyStruct->StructInstance);
 	}
 
-	PYCONVERSION_RETURN(Result, TEXT("Nativize"), *FString::Printf(TEXT("Cannot nativize '%s' as '%s'"), *PyUtil::GetFriendlyTypename(PyObj), *StructType->GetName()));
+	PYCONVERSION_RETURN(Result, TEXT("Nativize"), *FString::Printf(TEXT("Cannot nativize '%s' as '%s'"), *PyUtil::GetFriendlyTypename(PyObj), *PyUtil::GetFriendlyTypename(PyStructType)));
 }
 
 FPyConversionResult PythonizeStructInstance(UScriptStruct* StructType, const void* StructInstance, PyObject*& OutPyObj, const ESetErrorState SetErrorState)
@@ -627,6 +628,48 @@ PyObject* PythonizeStruct(UScriptStruct* Val, const ESetErrorState SetErrorState
 	return Obj;
 }
 
+FPyConversionResult NativizeEnumEntry(PyObject* PyObj, const UEnum* EnumType, int64& OutVal, const ESetErrorState SetErrorState)
+{
+	FPyConversionResult Result = FPyConversionResult::Failure();
+
+	PyTypeObject* PyEnumType = FPyWrapperTypeRegistry::Get().GetWrappedEnumType(EnumType);
+	FPyWrapperEnumPtr PyEnum = FPyWrapperEnumPtr::StealReference(FPyWrapperEnum::CastPyObject(PyObj, PyEnumType, &Result));
+	if (PyEnum)
+	{
+		OutVal = FPyWrapperEnum::GetEnumEntryValue(PyEnum);
+	}
+
+	PYCONVERSION_RETURN(Result, TEXT("Nativize"), *FString::Printf(TEXT("Cannot nativize '%s' as '%s'"), *PyUtil::GetFriendlyTypename(PyObj), *PyUtil::GetFriendlyTypename(PyEnumType)));
+}
+
+FPyConversionResult PythonizeEnumEntry(const int64 Val, const UEnum* EnumType, PyObject*& OutPyObj, const ESetErrorState SetErrorState)
+{
+	PyTypeObject* PyEnumType = FPyWrapperTypeRegistry::Get().GetWrappedEnumType(EnumType);
+	if (const FPyWrapperEnumMetaData* PyEnumMetaData = FPyWrapperEnumMetaData::GetMetaData(PyEnumType))
+	{
+		// Find an enum entry using this value
+		for (FPyWrapperEnum* PyEnumEntry : PyEnumMetaData->EnumEntries)
+		{
+			const int64 EnumEntryVal = FPyWrapperEnum::GetEnumEntryValue(PyEnumEntry);
+			if (EnumEntryVal == Val)
+			{
+				Py_INCREF(PyEnumEntry);
+				OutPyObj = (PyObject*)PyEnumEntry;
+				return FPyConversionResult::Success();
+			}
+		}
+	}
+
+	PYCONVERSION_RETURN(FPyConversionResult::Failure(), TEXT("Nativize"), *FString::Printf(TEXT("Cannot pythonize '%d' (int64) as '%s'"), Val, *PyUtil::GetFriendlyTypename(PyEnumType)));
+}
+
+PyObject* PythonizeEnumEntry(const int64 Val, const UEnum* EnumType, const ESetErrorState SetErrorState)
+{
+	PyObject* Obj = nullptr;
+	PythonizeEnumEntry(Val, EnumType, Obj, SetErrorState);
+	return Obj;
+}
+
 FPyConversionResult NativizeProperty(PyObject* PyObj, const UProperty* Prop, void* ValueAddr, const FPyWrapperOwnerContext& InChangeOwner, const ESetErrorState SetErrorState)
 {
 #define PYCONVERSION_PROPERTY_RETURN(RESULT) \
@@ -709,7 +752,6 @@ FPyConversionResult NativizeProperty_Direct(PyObject* PyObj, const UProperty* Pr
 
 	NATIVIZE_SETTER_PROPERTY(UBoolProperty);
 	NATIVIZE_INLINE_PROPERTY(UInt8Property);
-	NATIVIZE_INLINE_PROPERTY(UByteProperty);
 	NATIVIZE_INLINE_PROPERTY(UInt16Property);
 	NATIVIZE_INLINE_PROPERTY(UUInt16Property);
 	NATIVIZE_INLINE_PROPERTY(UIntProperty);
@@ -722,10 +764,74 @@ FPyConversionResult NativizeProperty_Direct(PyObject* PyObj, const UProperty* Pr
 	NATIVIZE_INLINE_PROPERTY(UNameProperty);
 	NATIVIZE_INLINE_PROPERTY(UTextProperty);
 
+	if (auto* CastProp = Cast<UByteProperty>(Prop))
+	{
+		uint8 NewValue = 0;
+		FPyConversionResult Result = FPyConversionResult::Failure();
+
+		if (CastProp->Enum)
+		{
+			int64 EnumVal = 0;
+			Result = NativizeEnumEntry(PyObj, CastProp->Enum, EnumVal, SetErrorState);
+			if (Result.GetState() == EPyConversionResultState::SuccessWithCoercion)
+			{
+				// Don't allow implicit conversion on enum properties
+				Result.SetState(EPyConversionResultState::Failure);
+			}
+			if (Result)
+			{
+				NewValue = (uint8)EnumVal;
+			}
+		}
+		else
+		{
+			Result = Nativize(PyObj, NewValue, SetErrorState);
+		}
+
+		if (Result)
+		{
+			auto* ValuePtr = static_cast<uint8*>(ValueAddr);
+			if (*ValuePtr != NewValue)
+			{
+				EmitPropertyChangeNotifications(InChangeOwner, [&]()
+				{
+					*ValuePtr = NewValue;
+				});
+			}
+		}
+		PYCONVERSION_PROPERTY_RETURN(Result);
+	}
+
 	if (auto* CastProp = Cast<UEnumProperty>(Prop))
 	{
+		FPyConversionResult Result = FPyConversionResult::Failure();
+
 		UNumericProperty* EnumInternalProp = CastProp->GetUnderlyingProperty();
-		PYCONVERSION_PROPERTY_RETURN(EnumInternalProp ? NativizeProperty_Direct(PyObj, EnumInternalProp, ValueAddr, InChangeOwner, SetErrorState) : FPyConversionResult::Failure());
+		if (EnumInternalProp)
+		{
+			int64 NewValue = 0;
+
+			Result = NativizeEnumEntry(PyObj, CastProp->GetEnum(), NewValue, SetErrorState);
+			if (Result.GetState() == EPyConversionResultState::SuccessWithCoercion)
+			{
+				// Don't allow implicit conversion on enum properties
+				Result.SetState(EPyConversionResultState::Failure);
+			}
+
+			if (Result)
+			{
+				const int64 OldValue = EnumInternalProp->GetSignedIntPropertyValue(ValueAddr);
+				if (OldValue != NewValue)
+				{
+					EmitPropertyChangeNotifications(InChangeOwner, [&]()
+					{
+						EnumInternalProp->SetIntPropertyValue(ValueAddr, NewValue);
+					});
+				}
+			}
+		}
+
+		PYCONVERSION_PROPERTY_RETURN(Result);
 	}
 
 	if (auto* CastProp = Cast<UClassProperty>(Prop))
@@ -912,7 +1018,6 @@ FPyConversionResult PythonizeProperty_Direct(const UProperty* Prop, const void*
 
 	PYTHONIZE_GETTER_PROPERTY(UBoolProperty);
 	PYTHONIZE_GETTER_PROPERTY(UInt8Property);
-	PYTHONIZE_GETTER_PROPERTY(UByteProperty);
 	PYTHONIZE_GETTER_PROPERTY(UInt16Property);
 	PYTHONIZE_GETTER_PROPERTY(UUInt16Property);
 	PYTHONIZE_GETTER_PROPERTY(UIntProperty);
@@ -925,10 +1030,23 @@ FPyConversionResult PythonizeProperty_Direct(const UProperty* Prop, const void*
 	PYTHONIZE_GETTER_PROPERTY(UNameProperty);
 	PYTHONIZE_GETTER_PROPERTY(UTextProperty);
 
+	if (auto* CastProp = Cast<UByteProperty>(Prop))
+	{
+		const uint8 Value = CastProp->GetPropertyValue(ValueAddr);
+		if (CastProp->Enum)
+		{
+			PYCONVERSION_PROPERTY_RETURN(PythonizeEnumEntry((int64)Value, CastProp->Enum, OutPyObj, SetErrorState));
+		}
+		else
+		{
+			PYCONVERSION_PROPERTY_RETURN(Pythonize(Value, OutPyObj, SetErrorState));
+		}
+	}
+
 	if (auto* CastProp = Cast<UEnumProperty>(Prop))
 	{
 		UNumericProperty* EnumInternalProp = CastProp->GetUnderlyingProperty();
-		PYCONVERSION_PROPERTY_RETURN(EnumInternalProp ? PythonizeProperty_Direct(EnumInternalProp, ValueAddr, OutPyObj, ConversionMethod, OwnerPyObj, SetErrorState) : FPyConversionResult::Failure());
+		PYCONVERSION_PROPERTY_RETURN(PythonizeEnumEntry(EnumInternalProp ? EnumInternalProp->GetSignedIntPropertyValue(ValueAddr) : 0, CastProp->GetEnum(), OutPyObj, SetErrorState));
 	}
 
 	if (auto* CastProp = Cast<UClassProperty>(Prop))
diff --git a/Engine/Plugins/Experimental/PythonScriptPlugin/Source/PythonScriptPlugin/Private/PyConversion.h b/Engine/Plugins/Experimental/PythonScriptPlugin/Source/PythonScriptPlugin/Private/PyConversion.h
--- a/Engine/Plugins/Experimental/PythonScriptPlugin/Source/PythonScriptPlugin/Private/PyConversion.h
+++ b/Engine/Plugins/Experimental/PythonScriptPlugin/Source/PythonScriptPlugin/Private/PyConversion.h
@@ -100,6 +100,11 @@ namespace PyConversion
 	FPyConversionResult PythonizeStruct(UScriptStruct* Val, PyObject*& OutPyObj, const ESetErrorState SetErrorState = ESetErrorState::Yes);
 	PyObject* PythonizeStruct(UScriptStruct* Val, const ESetErrorState SetErrorState = ESetErrorState::Yes);
 
+	/** Conversion for enum entries */
+	FPyConversionResult NativizeEnumEntry(PyObject* PyObj, const UEnum* EnumType, int64& OutVal, const ESetErrorState SetErrorState = ESetErrorState::Yes);
+	FPyConversionResult PythonizeEnumEntry(const int64 Val, const UEnum* EnumType, PyObject*& OutPyObj, const ESetErrorState SetErrorState = ESetErrorState::Yes);
+	PyObject* PythonizeEnumEntry(const int64 Val, const UEnum* EnumType, const ESetErrorState SetErrorState = ESetErrorState::Yes);
+
 	namespace Internal
 	{
 		/** Internal version of NativizeStructInstance/PythonizeStructInstance that work on the type-erased data */
@@ -193,16 +198,23 @@ namespace PyConversion
 
 	/** Conversion for known enum types */
 	template <typename T>
-	FPyConversionResult NativizeEnum(PyObject* PyObj, T& OutVal, const ESetErrorState SetErrorState = ESetErrorState::Yes)
+	FPyConversionResult NativizeEnumEntry(PyObject* PyObj, const UEnum* EnumType, T& OutVal, const ESetErrorState SetErrorState = ESetErrorState::Yes)
 	{
-		return Nativize(PyObj, (__underlying_type(T)&)OutVal, SetErrorState);
+		int64 OutTmpVal = 0;
+		FPyConversionResult Result = NativizeEnumEntry(PyObj, EnumType, OutTmpVal, SetErrorState);
+		if (Result)
+		{
+			OutVal = (T)OutTmpVal;
+		}
+		return Result;
 	}
 
 	/** Conversion for known enum types */
 	template <typename T>
-	FPyConversionResult PythonizeEnum(const T& Val, PyObject*& OutPyObj, const ESetErrorState SetErrorState = ESetErrorState::Yes)
+	FPyConversionResult PythonizeEnumEntry(const T& Val, const UEnum* EnumType, PyObject*& OutPyObj, const ESetErrorState SetErrorState = ESetErrorState::Yes)
 	{
-		return Pythonize((__underlying_type(T))Val, OutPyObj, SetErrorState);
+		const int64 TmpVal = (int64)Val;
+		return PythonizeEnumEntry(TmpVal, EnumType, OutPyObj, SetErrorState);
 	}
 
 	/** Conversion for property instances (including fixed arrays) - ValueAddr should point to the property data */
diff --git a/Engine/Plugins/Experimental/PythonScriptPlugin/Source/PythonScriptPlugin/Private/PyFileWriter.cpp b/Engine/Plugins/Experimental/PythonScriptPlugin/Source/PythonScriptPlugin/Private/PyFileWriter.cpp
--- a/Engine/Plugins/Experimental/PythonScriptPlugin/Source/PythonScriptPlugin/Private/PyFileWriter.cpp
+++ b/Engine/Plugins/Experimental/PythonScriptPlugin/Source/PythonScriptPlugin/Private/PyFileWriter.cpp
@@ -49,7 +49,7 @@ void FPyFileWriter::Write(const FString& InStr)
 
 void FPyFileWriter::WriteDocString(const TCHAR* InDocString)
 {
-	if (*InDocString)
+	if (InDocString && *InDocString)
 	{
 		WriteDocString(FString(InDocString));
 	}
diff --git a/Engine/Plugins/Experimental/PythonScriptPlugin/Source/PythonScriptPlugin/Private/PyGenUtil.cpp b/Engine/Plugins/Experimental/PythonScriptPlugin/Source/PythonScriptPlugin/Private/PyGenUtil.cpp
--- a/Engine/Plugins/Experimental/PythonScriptPlugin/Source/PythonScriptPlugin/Private/PyGenUtil.cpp
+++ b/Engine/Plugins/Experimental/PythonScriptPlugin/Source/PythonScriptPlugin/Private/PyGenUtil.cpp
@@ -653,24 +653,34 @@ void FGeneratedWrappedEnumType::Finalize_PostReady()
 {
 	FGeneratedWrappedType::Finalize_PostReady();
 
+	check(MetaData.IsValid() && MetaData->GetTypeId() == FPyWrapperEnumMetaData::StaticTypeId());
+	TSharedRef<FPyWrapperEnumMetaData> EnumMetaData = StaticCastSharedRef<FPyWrapperEnumMetaData>(MetaData.ToSharedRef());
+
 	// Execute Python code within this block
 	{
 		FPyScopedGIL GIL;
 		for (const FGeneratedWrappedEnumEntry& EnumEntry : EnumEntries)
 		{
-			FPyWrapperEnum::SetEnumEntryValue(&PyType, EnumEntry.EntryValue, EnumEntry.EntryName.GetData(), EnumEntry.EntryDoc.GetData());
+			FPyWrapperEnum* PyEnumEntry = FPyWrapperEnum::AddEnumEntry(&PyType, EnumEntry.EntryValue, EnumEntry.EntryName.GetData(), EnumEntry.EntryDoc.GetData());
+			if (PyEnumEntry)
+			{
+				EnumMetaData->EnumEntries.Add(PyEnumEntry);
+			}
 		}
 	}
+
+	EnumMetaData->bFinalized = true;
 }
 
 void FGeneratedWrappedEnumType::ExtractEnumEntries(const UEnum* InEnum)
 {
 	for (int32 EnumEntryIndex = 0; EnumEntryIndex < InEnum->NumEnums() - 1; ++EnumEntryIndex)
 	{
+		// todo: deprecated enum entries?
 		if (ShouldExportEnumEntry(InEnum, EnumEntryIndex))
 		{
 			FGeneratedWrappedEnumEntry& EnumEntry = EnumEntries.AddDefaulted_GetRef();
-			EnumEntry.EntryName = TCHARToUTF8Buffer(*PythonizeName(InEnum->GetNameStringByIndex(EnumEntryIndex), EPythonizeNameCase::Upper));
+			EnumEntry.EntryName = TCHARToUTF8Buffer(*GetEnumEntryPythonName(InEnum, EnumEntryIndex));
 			EnumEntry.EntryDoc = TCHARToUTF8Buffer(*PythonizeTooltip(GetEnumEntryTooltip(InEnum, EnumEntryIndex)));
 			EnumEntry.EntryValue = InEnum->GetValueByIndex(EnumEntryIndex);
 		}
@@ -2326,6 +2336,34 @@ TArray<FString> GetDeprecatedEnumPythonNames(const UEnum* InEnum)
 	return GetDeprecatedFieldPythonNamesImpl(InEnum, ScriptNameMetaDataKey);
 }
 
+FString GetEnumEntryPythonName(const UEnum* InEnum, const int32 InEntryIndex)
+{
+	FString EnumEntryName;
+
+	// First see if we have a name override in the meta-data
+	{
+		EnumEntryName = InEnum->GetMetaData(TEXT("ScriptName"), InEntryIndex);
+
+		// This may be a semi-colon separated list - the first item is the one we want for the current name
+		if (!EnumEntryName.IsEmpty())
+		{
+			int32 SemiColonIndex = INDEX_NONE;
+			if (EnumEntryName.FindChar(TEXT(';'), SemiColonIndex))
+			{
+				EnumEntryName.RemoveAt(SemiColonIndex, EnumEntryName.Len() - SemiColonIndex, /*bAllowShrinking*/false);
+			}
+		}
+	}
+	
+	// Just use the entry name if we have no meta-data
+	if (EnumEntryName.IsEmpty())
+	{
+		EnumEntryName = InEnum->GetNameStringByIndex(InEntryIndex);
+	}
+
+	return PythonizeName(EnumEntryName, EPythonizeNameCase::Upper);
+}
+
 FString GetDelegatePythonName(const UFunction* InDelegateSignature)
 {
 	return InDelegateSignature->GetName().LeftChop(19); // Trim the "__DelegateSignature" suffix from the name
diff --git a/Engine/Plugins/Experimental/PythonScriptPlugin/Source/PythonScriptPlugin/Private/PyGenUtil.h b/Engine/Plugins/Experimental/PythonScriptPlugin/Source/PythonScriptPlugin/Private/PyGenUtil.h
--- a/Engine/Plugins/Experimental/PythonScriptPlugin/Source/PythonScriptPlugin/Private/PyGenUtil.h
+++ b/Engine/Plugins/Experimental/PythonScriptPlugin/Source/PythonScriptPlugin/Private/PyGenUtil.h
@@ -1003,6 +1003,9 @@ namespace PyGenUtil
 	/** Get the deprecated Python names of the given enum */
 	TArray<FString> GetDeprecatedEnumPythonNames(const UEnum* InEnum);
 
+	/** Get the Python name of the given enum entry */
+	FString GetEnumEntryPythonName(const UEnum* InEnum, const int32 InEntryIndex);
+
 	/** Get the Python name of the given delegate signature */
 	FString GetDelegatePythonName(const UFunction* InDelegateSignature);
 
diff --git a/Engine/Plugins/Experimental/PythonScriptPlugin/Source/PythonScriptPlugin/Private/PySlate.cpp b/Engine/Plugins/Experimental/PythonScriptPlugin/Source/PythonScriptPlugin/Private/PySlate.cpp
--- a/Engine/Plugins/Experimental/PythonScriptPlugin/Source/PythonScriptPlugin/Private/PySlate.cpp
+++ b/Engine/Plugins/Experimental/PythonScriptPlugin/Source/PythonScriptPlugin/Private/PySlate.cpp
@@ -8,6 +8,7 @@
 #include "PyConversion.h"
 #include "PyWrapperTypeRegistry.h"
 
+#include "UObject/Package.h"
 #include "Framework/Application/SlateApplication.h"
 
 #if PLATFORM_WINDOWS
@@ -139,8 +140,9 @@ PyObject* ParentExternalWindowToSlate(PyObject* InSelf, PyObject* InArgs)
 		return nullptr;
 	}
 
+	static const UEnum* ParentWindowSearchMethodEnum = FindObject<UEnum>(ANY_PACKAGE, TEXT("ESlateParentWindowSearchMethod"));
 	ESlateParentWindowSearchMethod ParentWindowSearchMethod = ESlateParentWindowSearchMethod::ActiveWindow;
-	if (PyParentWindowSearchMethod && !PyConversion::NativizeEnum(PyParentWindowSearchMethod, ParentWindowSearchMethod))
+	if (PyParentWindowSearchMethod && !PyConversion::NativizeEnumEntry(PyParentWindowSearchMethod, ParentWindowSearchMethodEnum, ParentWindowSearchMethod))
 	{
 		PyUtil::SetPythonError(PyExc_TypeError, TEXT("parent_external_window_to_slate"), *FString::Printf(TEXT("Failed to convert argument '%s' to 'SlateParentWindowSearchMethod'"), *PyUtil::GetFriendlyTypename(PyParentWindowSearchMethod)));
 		return nullptr;
diff --git a/Engine/Plugins/Experimental/PythonScriptPlugin/Source/PythonScriptPlugin/Private/PyWrapperEnum.cpp b/Engine/Plugins/Experimental/PythonScriptPlugin/Source/PythonScriptPlugin/Private/PyWrapperEnum.cpp
--- a/Engine/Plugins/Experimental/PythonScriptPlugin/Source/PythonScriptPlugin/Private/PyWrapperEnum.cpp
+++ b/Engine/Plugins/Experimental/PythonScriptPlugin/Source/PythonScriptPlugin/Private/PyWrapperEnum.cpp
@@ -6,39 +6,112 @@
 #include "PyUtil.h"
 #include "PyConversion.h"
 
+#include "Containers/ArrayView.h"
+
 #if WITH_PYTHON
 
-/** Python type for FPyWrapperEnumValueDescrObject */
-extern PyTypeObject PyWrapperEnumValueDescrType;
+typedef TArrayView<FPyWrapperEnum* const> FPyWrapperEnumArrayView;
 
-/** Python object for the descriptor of an enum value */
-struct FPyWrapperEnumValueDescrObject
+/** Python metaclass type for FPyWrapperEnum */
+extern PyTypeObject PyWrapperEnumMetaclassType;
+
+/** Python type for FPyWrapperEnumIterator */
+extern PyTypeObject PyWrapperEnumIteratorType;
+
+/** Iterator used with enums */
+struct FPyWrapperEnumIterator
 {
 	/** Common Python Object */
 	PyObject_HEAD
 
-	/** The owner enum type */
-	PyTypeObject* EnumType;
+	/** Array being iterated over */
+	FPyWrapperEnumArrayView IterArray;
+
+	/** Current iteration index */
+	int32 IterIndex;
+
+	/** New this iterator instance (called via tp_new for Python, or directly in C++) */
+	static FPyWrapperEnumIterator* New(PyTypeObject* InType)
+	{
+		FPyWrapperEnumIterator* Self = (FPyWrapperEnumIterator*)InType->tp_alloc(InType, 0);
+		if (Self)
+		{
+			new(&Self->IterArray) FPyWrapperEnumArrayView();
+			Self->IterIndex = 0;
+		}
+		return Self;
+	}
+
+	/** Free this iterator instance (called via tp_dealloc for Python) */
+	static void Free(FPyWrapperEnumIterator* InSelf)
+	{
+		Deinit(InSelf);
 
-	/** The enum value */
-	int64 EnumValue;
+		InSelf->IterArray.~FPyWrapperEnumArrayView();
+		Py_TYPE(InSelf)->tp_free((PyObject*)InSelf);
+	}
 
-	/** The enum value name */
-	PyObject* EnumValueName;
+	/** Initialize this iterator instance (called via tp_init for Python, or directly in C++) */
+	static int Init(FPyWrapperEnumIterator* InSelf, FPyWrapperEnumArrayView InIterArray)
+	{
+		Deinit(InSelf);
 
-	/** The enum value doc string (may be null) */
-	PyObject* EnumValueDoc;
+		InSelf->IterArray = InIterArray;
+		InSelf->IterIndex = 0;
+
+		return 0;
+	}
+
+	/** Deinitialize this iterator instance (called via Init and Free to restore the instance to its New state) */
+	static void Deinit(FPyWrapperEnumIterator* InSelf)
+	{
+		InSelf->IterArray = FPyWrapperEnumArrayView();
+		InSelf->IterIndex = 0;
+	}
+
+	/** Get the iterator */
+	static FPyWrapperEnumIterator* GetIter(FPyWrapperEnumIterator* InSelf)
+	{
+		Py_INCREF(InSelf);
+		return InSelf;
+	}
+
+	/** Return the current value (if any) and advance the iterator */
+	static PyObject* IterNext(FPyWrapperEnumIterator* InSelf)
+	{
+		if (InSelf->IterArray.IsValidIndex(InSelf->IterIndex))
+		{
+			FPyWrapperEnum* EnumEntry = InSelf->IterArray[InSelf->IterIndex++];
+			Py_INCREF(EnumEntry);
+			return (PyObject*)EnumEntry;
+		}
+
+		PyErr_SetObject(PyExc_StopIteration, Py_None);
+		return nullptr;
+	}
+};
+
+/** Python object for the descriptor of an enum value */
+struct FPyWrapperEnumValueDescrObject
+{
+	/** Common Python Object */
+	PyObject_HEAD
+
+	/** The enum entry */
+	FPyWrapperEnum* EnumEntry;
+
+	/** The enum entry doc string (may be null) */
+	PyObject* EnumEntryDoc;
 
 	/** New an instance */
-	static FPyWrapperEnumValueDescrObject* New(PyTypeObject* InEnumType, const int64 InEnumValue, const char* InEnumValueName, const char* InEnumValueDoc)
+	static FPyWrapperEnumValueDescrObject* New(PyTypeObject* InEnumType, const int64 InEnumEntryValue, const char* InEnumEntryName, const char* InEnumEntryDoc)
 	{
 		FPyWrapperEnumValueDescrObject* Self = (FPyWrapperEnumValueDescrObject*)PyWrapperEnumValueDescrType.tp_alloc(&PyWrapperEnumValueDescrType, 0);
 		if (Self)
 		{
-			Self->EnumType = InEnumType;
-			Self->EnumValue = InEnumValue;
-			Self->EnumValueName = PyUnicode_FromString(InEnumValueName);
-			Self->EnumValueDoc = InEnumValueDoc ? PyUnicode_FromString(InEnumValueDoc) : nullptr;
+			Self->EnumEntry = FPyWrapperEnum::New(InEnumType);
+			FPyWrapperEnum::Init(Self->EnumEntry, InEnumEntryValue, InEnumEntryName);
+			Self->EnumEntryDoc = InEnumEntryDoc ? PyUnicode_FromString(InEnumEntryDoc) : nullptr;
 		}
 		return Self;
 	}
@@ -46,20 +119,25 @@ struct FPyWrapperEnumValueDescrObject
 	/** Free this instance */
 	static void Free(FPyWrapperEnumValueDescrObject* InSelf)
 	{
-		InSelf->EnumType = nullptr;
+		Py_XDECREF(InSelf->EnumEntry);
+		InSelf->EnumEntry = nullptr;
 
-		Py_XDECREF(InSelf->EnumValueName);
-		InSelf->EnumValueName = nullptr;
-
-		Py_XDECREF(InSelf->EnumValueDoc);
-		InSelf->EnumValueDoc = nullptr;
+		Py_XDECREF(InSelf->EnumEntryDoc);
+		InSelf->EnumEntryDoc = nullptr;
 
 		Py_TYPE(InSelf)->tp_free((PyObject*)InSelf);
 	}
 };
 
+typedef TPyPtr<FPyWrapperEnumValueDescrObject> FPyWrapperEnumValueDescrObjectPtr;
+
 void InitializePyWrapperEnum(PyGenUtil::FNativePythonModule& ModuleInfo)
 {
+	PyType_Ready(&PyWrapperEnumIteratorType);
+	PyType_Ready(&PyWrapperEnumMetaclassType);
+
+	// Set the metaclass on the enum type
+	Py_TYPE(&PyWrapperEnumType) = &PyWrapperEnumMetaclassType;
 	if (PyType_Ready(&PyWrapperEnumType) == 0)
 	{
 		static FPyWrapperEnumMetaData MetaData;
@@ -67,7 +145,28 @@ void InitializePyWrapperEnum(PyGenUtil::FNativePythonModule& ModuleInfo)
 		ModuleInfo.AddType(&PyWrapperEnumType);
 	}
 
-	PyType_Ready(&PyWrapperEnumValueDescrType);
+	if (PyType_Ready(&PyWrapperEnumValueDescrType) == 0)
+	{
+		ModuleInfo.AddType(&PyWrapperEnumValueDescrType);
+	}
+}
+
+FPyWrapperEnum* FPyWrapperEnum::New(PyTypeObject* InType)
+{
+	FPyWrapperEnum* Self = (FPyWrapperEnum*)FPyWrapperBase::New(InType);
+	if (Self)
+	{
+		Self->EntryName = nullptr;
+		Self->EntryValue = nullptr;
+	}
+	return Self;
+}
+
+void FPyWrapperEnum::Free(FPyWrapperEnum* InSelf)
+{
+	Deinit(InSelf);
+
+	FPyWrapperBase::Free(InSelf);
 }
 
 int FPyWrapperEnum::Init(FPyWrapperEnum* InSelf)
@@ -76,27 +175,228 @@ int FPyWrapperEnum::Init(FPyWrapperEnum* InSelf)
 	return -1;
 }
 
-void FPyWrapperEnum::SetEnumEntryValue(PyTypeObject* InType, const int64 InEnumValue, const char* InEnumValueName, const char* InEnumValueDoc)
+int FPyWrapperEnum::Init(FPyWrapperEnum* InSelf, const int64 InEnumEntryValue, const char* InEnumEntryName)
+{
+	if (FPyWrapperEnumMetaData::IsEnumFinalized(InSelf))
+	{
+		PyUtil::SetPythonError(PyExc_Exception, InSelf, TEXT("Cannot create instances of enum types"));
+		return -1;
+	}
+
+	InSelf->EntryName = PyUnicode_FromString(InEnumEntryName);
+	InSelf->EntryValue = PyConversion::Pythonize(InEnumEntryValue);
+	return 0;
+}
+
+void FPyWrapperEnum::Deinit(FPyWrapperEnum* InSelf)
+{
+	Py_XDECREF(InSelf->EntryName);
+	InSelf->EntryName = nullptr;
+
+	Py_XDECREF(InSelf->EntryValue);
+	InSelf->EntryValue = nullptr;
+}
+
+bool FPyWrapperEnum::ValidateInternalState(FPyWrapperEnum* InSelf)
+{
+	if (!InSelf->EntryName)
+	{
+		PyUtil::SetPythonError(PyExc_Exception, Py_TYPE(InSelf), TEXT("Internal Error - EntryName is null!"));
+		return false;
+	}
+
+	if (!InSelf->EntryValue)
+	{
+		PyUtil::SetPythonError(PyExc_Exception, Py_TYPE(InSelf), TEXT("Internal Error - EntryValue is null!"));
+		return false;
+	}
+
+	return true;
+}
+
+FPyWrapperEnum* FPyWrapperEnum::CastPyObject(PyObject* InPyObject, FPyConversionResult* OutCastResult)
 {
-	FPyObjectPtr PyEnumValueDescr = FPyObjectPtr::StealReference((PyObject*)FPyWrapperEnumValueDescrObject::New(InType, InEnumValue, InEnumValueName, InEnumValueDoc));
-	if (PyEnumValueDescr)
+	SetOptionalPyConversionResult(FPyConversionResult::Failure(), OutCastResult);
+
+	if (PyObject_IsInstance(InPyObject, (PyObject*)&PyWrapperEnumType) == 1)
 	{
-		PyDict_SetItemString(InType->tp_dict, InEnumValueName, PyEnumValueDescr);
+		SetOptionalPyConversionResult(FPyConversionResult::Success(), OutCastResult);
+
+		Py_INCREF(InPyObject);
+		return (FPyWrapperEnum*)InPyObject;
 	}
+
+	return nullptr;
+}
+
+FPyWrapperEnum* FPyWrapperEnum::CastPyObject(PyObject* InPyObject, PyTypeObject* InType, FPyConversionResult* OutCastResult)
+{
+	SetOptionalPyConversionResult(FPyConversionResult::Failure(), OutCastResult);
+
+	if (PyObject_IsInstance(InPyObject, (PyObject*)InType) == 1 && (InType == &PyWrapperEnumType || PyObject_IsInstance(InPyObject, (PyObject*)&PyWrapperEnumType) == 1))
+	{
+		SetOptionalPyConversionResult(Py_TYPE(InPyObject) == InType ? FPyConversionResult::Success() : FPyConversionResult::SuccessWithCoercion(), OutCastResult);
+
+		Py_INCREF(InPyObject);
+		return (FPyWrapperEnum*)InPyObject;
+	}
+
+	// Allow casting from a different enum type using the same UEnum (for deprecation)
+	if (PyObject_IsInstance(InPyObject, (PyObject*)&PyWrapperEnumType) == 1)
+	{
+		const UEnum* RequiredEnum = FPyWrapperEnumMetaData::GetEnum(InType);
+		const UEnum* ActualEnum = FPyWrapperEnumMetaData::GetEnum(Py_TYPE(InPyObject));
+
+		if (RequiredEnum == ActualEnum)
+		{
+			SetOptionalPyConversionResult(FPyConversionResult::Success(), OutCastResult);
+
+			Py_INCREF(InPyObject);
+			return (FPyWrapperEnum*)InPyObject;
+		}
+	}
+
+	// Allow coerced casting from a numeric value
+	if (const FPyWrapperEnumMetaData* EnumMetaData = FPyWrapperEnumMetaData::GetMetaData(InType))
+	{
+		int64 OtherVal = 0;
+		if (PyConversion::Nativize(InPyObject, OtherVal, PyConversion::ESetErrorState::No))
+		{
+			// Find an enum entry using this value
+			for (FPyWrapperEnum* PyEnumEntry : EnumMetaData->EnumEntries)
+			{
+				const int64 EnumEntryVal = FPyWrapperEnum::GetEnumEntryValue(PyEnumEntry);
+				if (EnumEntryVal == OtherVal)
+				{
+					SetOptionalPyConversionResult(FPyConversionResult::SuccessWithCoercion(), OutCastResult);
+
+					Py_INCREF(PyEnumEntry);
+					return (FPyWrapperEnum*)PyEnumEntry;
+				}
+			}
+		}
+	}
+
+	return nullptr;
+}
+
+FString FPyWrapperEnum::GetEnumEntryName(FPyWrapperEnum* InSelf)
+{
+	FString EnumEntryNameStr;
+	if (InSelf->EntryName)
+	{
+		PyConversion::Nativize(InSelf->EntryName, EnumEntryNameStr, PyConversion::ESetErrorState::No);
+	}
+	return EnumEntryNameStr;
+}
+
+int64 FPyWrapperEnum::GetEnumEntryValue(FPyWrapperEnum* InSelf)
+{
+	int64 EnumEntryValueNum = 0;
+	if (InSelf->EntryValue)
+	{
+		PyConversion::Nativize(InSelf->EntryValue, EnumEntryValueNum, PyConversion::ESetErrorState::No);
+	}
+	return EnumEntryValueNum;
+}
+
+FPyWrapperEnum* FPyWrapperEnum::AddEnumEntry(PyTypeObject* InType, const int64 InEnumEntryValue, const char* InEnumEntryName, const char* InEnumEntryDoc)
+{
+	if (!FPyWrapperEnumMetaData::IsEnumFinalized(InType))
+	{
+		FPyWrapperEnumValueDescrObjectPtr PyEnumValueDescr = FPyWrapperEnumValueDescrObjectPtr::StealReference(FPyWrapperEnumValueDescrObject::New(InType, InEnumEntryValue, InEnumEntryName, InEnumEntryDoc));
+		if (PyEnumValueDescr)
+		{
+			PyDict_SetItemString(InType->tp_dict, InEnumEntryName, (PyObject*)PyEnumValueDescr.GetPtr());
+			return PyEnumValueDescr->EnumEntry;
+		}
+	}
+	return nullptr;
 }
 
 PyTypeObject InitializePyWrapperEnumType()
 {
 	struct FFuncs
 	{
+		static PyObject* New(PyTypeObject* InType, PyObject* InArgs, PyObject* InKwds)
+		{
+			return (PyObject*)FPyWrapperEnum::New(InType);
+		}
+
+		static void Dealloc(FPyWrapperEnum* InSelf)
+		{
+			FPyWrapperEnum::Free(InSelf);
+		}
+
 		static int Init(FPyWrapperEnum* InSelf, PyObject* InArgs, PyObject* InKwds)
 		{
 			return FPyWrapperEnum::Init(InSelf);
 		}
+
+		static PyObject* Str(FPyWrapperEnum* InSelf)
+		{
+			if (!FPyWrapperEnum::ValidateInternalState(InSelf))
+			{
+				return nullptr;
+			}
+
+			return PyUnicode_FromFormat("<%s.%s: %d>", Py_TYPE(InSelf)->tp_name, TCHAR_TO_UTF8(*FPyWrapperEnum::GetEnumEntryName(InSelf)), FPyWrapperEnum::GetEnumEntryValue(InSelf));
+		}
+
+		static PyObject* RichCmp(FPyWrapperEnum* InSelf, PyObject* InOther, int InOp)
+		{
+			if (!FPyWrapperEnum::ValidateInternalState(InSelf))
+			{
+				return nullptr;
+			}
+
+			FPyWrapperEnumPtr Other = FPyWrapperEnumPtr::StealReference(FPyWrapperEnum::CastPyObject(InOther, Py_TYPE(InSelf)));
+			if (!Other)
+			{
+				Py_INCREF(Py_NotImplemented);
+				return Py_NotImplemented;
+			}
+
+			if (InOp != Py_EQ && InOp != Py_NE)
+			{
+				PyUtil::SetPythonError(PyExc_TypeError, InSelf, TEXT("Only == and != comparison is supported"));
+				return nullptr;
+			}
+
+			// Compare the objects if both enums are the same type, otherwise compare the values (as cast must have returned a deprecated enum and the entry objects won't match)
+			const bool bIsIdentical = (Py_TYPE(InSelf) == Py_TYPE(Other.GetPtr())) ? InSelf->EntryValue == Other->EntryValue : FPyWrapperEnum::GetEnumEntryValue(InSelf) == FPyWrapperEnum::GetEnumEntryValue(Other);
+			return PyBool_FromLong(InOp == Py_EQ ? bIsIdentical : !bIsIdentical);
+		}
+
+		static PyUtil::FPyHashType Hash(FPyWrapperEnum* InSelf)
+		{
+			if (!FPyWrapperEnum::ValidateInternalState(InSelf))
+			{
+				return -1;
+			}
+
+			return PyObject_Hash(InSelf->EntryValue);
+		}
 	};
 
 	struct FMethods
 	{
+		static PyObject* Cast(PyTypeObject* InType, PyObject* InArgs)
+		{
+			PyObject* PyObj = nullptr;
+			if (PyArg_ParseTuple(InArgs, "O:cast", &PyObj))
+			{
+				PyObject* PyCastResult = (PyObject*)FPyWrapperEnum::CastPyObject(PyObj, InType);
+				if (!PyCastResult)
+				{
+					PyUtil::SetPythonError(PyExc_TypeError, InType, *FString::Printf(TEXT("Cannot cast type '%s' to '%s'"), *PyUtil::GetFriendlyTypename(PyObj), *PyUtil::GetFriendlyTypename(InType)));
+				}
+				return PyCastResult;
+			}
+
+			return nullptr;
+		}
+
 		static PyObject* StaticEnum(PyTypeObject* InType)
 		{
 			UEnum* Enum = FPyWrapperEnumMetaData::GetEnum(InType);
@@ -104,8 +404,15 @@ PyTypeObject InitializePyWrapperEnumType()
 		}
 	};
 
+	static PyMemberDef PyMembers[] = {
+		{ PyCStrCast("name"), T_OBJECT, STRUCT_OFFSET(FPyWrapperEnum, EntryName), READONLY, PyCStrCast("The name of this enum entry") },
+		{ PyCStrCast("value"), T_OBJECT, STRUCT_OFFSET(FPyWrapperEnum, EntryValue), READONLY, PyCStrCast("The numeric value of this enum entry") },
+		{ nullptr, 0, 0, 0, nullptr }
+	};
+
 	static PyMethodDef PyMethods[] = {
-		{ "static_enum", PyCFunctionCast(&FMethods::StaticEnum), METH_NOARGS | METH_CLASS, "X.static_enum() -> UEnum -- get the Unreal enum of this type" },
+		{ "cast", PyCFunctionCast(&FMethods::Cast), METH_VARARGS | METH_CLASS, "X.cast(object) -> enum -- cast the given object to this Unreal enum type" },
+		{ "static_enum", PyCFunctionCast(&FMethods::StaticEnum), METH_NOARGS | METH_CLASS, "X.static_enum() -> Enum -- get the Unreal enum of this type" },
 		{ nullptr, nullptr, 0, nullptr }
 	};
 
@@ -116,8 +423,14 @@ PyTypeObject InitializePyWrapperEnumType()
 	};
 
 	PyType.tp_base = &PyWrapperBaseType;
+	PyType.tp_new = (newfunc)&FFuncs::New;
+	PyType.tp_dealloc = (destructor)&FFuncs::Dealloc;
 	PyType.tp_init = (initproc)&FFuncs::Init;
+	PyType.tp_str = (reprfunc)&FFuncs::Str;
+	PyType.tp_richcompare = (richcmpfunc)&FFuncs::RichCmp;
+	PyType.tp_hash = (hashfunc)&FFuncs::Hash;
 
+	PyType.tp_members = PyMembers;
 	PyType.tp_methods = PyMethods;
 
 	PyType.tp_flags = Py_TPFLAGS_DEFAULT | Py_TPFLAGS_BASETYPE;
@@ -142,11 +455,16 @@ PyTypeObject InitializePyWrapperEnumValueDescrType()
 
 		static PyObject* DescrGet(FPyWrapperEnumValueDescrObject* InSelf, PyObject* InObj, PyObject* InType)
 		{
-			// Deprecated structs emit a warning
+			if (!FPyWrapperEnum::ValidateInternalState(InSelf->EnumEntry))
+			{
+				return nullptr;
+			}
+
+			// Deprecated enums emit a warning
 			{
 				FString DeprecationMessage;
-				if (FPyWrapperEnumMetaData::IsEnumDeprecated(InSelf->EnumType, &DeprecationMessage) &&
-					PyUtil::SetPythonWarning(PyExc_DeprecationWarning, InSelf->EnumType, *FString::Printf(TEXT("Enum '%s' is deprecated: %s"), UTF8_TO_TCHAR(InSelf->EnumType->tp_name), *DeprecationMessage)) == -1
+				if (FPyWrapperEnumMetaData::IsEnumDeprecated(InSelf->EnumEntry, &DeprecationMessage) &&
+					PyUtil::SetPythonWarning(PyExc_DeprecationWarning, InSelf->EnumEntry, *FString::Printf(TEXT("Enum '%s' is deprecated: %s"), UTF8_TO_TCHAR(Py_TYPE(InSelf->EnumEntry)->tp_name), *DeprecationMessage)) == -1
 					)
 				{
 					// -1 from SetPythonWarning means the warning should be an exception
@@ -154,7 +472,8 @@ PyTypeObject InitializePyWrapperEnumValueDescrType()
 				}
 			}
 
-			return PyConversion::Pythonize(InSelf->EnumValue);
+			Py_INCREF(InSelf->EnumEntry);
+			return (PyObject*)InSelf->EnumEntry;
 		}
 
 		static int DescrSet(FPyWrapperEnumValueDescrObject* InSelf, PyObject* InObj, PyObject* InValue)
@@ -166,31 +485,38 @@ PyTypeObject InitializePyWrapperEnumValueDescrType()
 
 	struct FGetSets
 	{
-		static PyObject* GetDoc(FPyWrapperEnumValueDescrObject* InSelf, void* InClosure)
+		static PyObject* GetName(FPyWrapperEnumValueDescrObject* InSelf, void* InClosure)
 		{
-			if (InSelf->EnumValueDoc)
+			if (InSelf->EnumEntry && InSelf->EnumEntry->EntryName)
 			{
-				Py_INCREF(InSelf->EnumValueDoc);
-				return InSelf->EnumValueDoc;
+				Py_INCREF(InSelf->EnumEntry->EntryName);
+				return InSelf->EnumEntry->EntryName;
 			}
 
 			Py_RETURN_NONE;
 		}
-	};
 
-	static PyMemberDef PyMembers[] = {
-		{ PyCStrCast("__name__"), T_OBJECT, STRUCT_OFFSET(FPyWrapperEnumValueDescrObject, EnumValueName), READONLY, nullptr },
-		{ nullptr, 0, 0, 0, nullptr }
+		static PyObject* GetDoc(FPyWrapperEnumValueDescrObject* InSelf, void* InClosure)
+		{
+			if (InSelf->EnumEntryDoc)
+			{
+				Py_INCREF(InSelf->EnumEntryDoc);
+				return InSelf->EnumEntryDoc;
+			}
+
+			Py_RETURN_NONE;
+		}
 	};
 
 	static PyGetSetDef PyGetSets[] = {
+		{ PyCStrCast("__name__"), (getter)&FGetSets::GetName, nullptr, nullptr, nullptr },
 		{ PyCStrCast("__doc__"), (getter)&FGetSets::GetDoc, nullptr, nullptr, nullptr },
 		{ nullptr, nullptr, nullptr, nullptr, nullptr }
 	};
 
 	PyTypeObject PyType = {
 		PyVarObject_HEAD_INIT(nullptr, 0)
-		"enum_value", /* tp_name */
+		"_EnumEntry", /* tp_name */
 		sizeof(FPyWrapperEnumValueDescrObject), /* tp_basicsize */
 	};
 
@@ -200,7 +526,6 @@ PyTypeObject InitializePyWrapperEnumValueDescrType()
 	PyType.tp_descr_set = (descrsetfunc)&FFuncs::DescrSet;
 	PyType.tp_getattro = (getattrofunc)&PyObject_GenericGetAttr;
 
-	PyType.tp_members = PyMembers;
 	PyType.tp_getset = PyGetSets;
 
 	PyType.tp_flags = Py_TPFLAGS_DEFAULT;
@@ -208,11 +533,136 @@ PyTypeObject InitializePyWrapperEnumValueDescrType()
 	return PyType;
 }
 
+PyTypeObject InitializePyWrapperEnumMetaclassType()
+{
+	struct FFuncs
+	{
+		static PyObject* GetIter(PyTypeObject* InSelf)
+		{
+			typedef TPyPtr<FPyWrapperEnumIterator> FPyWrapperEnumIteratorPtr;
+
+			FPyWrapperEnumArrayView EnumEntriesArray;
+			if (const FPyWrapperEnumMetaData* EnumMetaData = FPyWrapperEnumMetaData::GetMetaData(InSelf))
+			{
+				EnumEntriesArray = EnumMetaData->EnumEntries;
+			}
+			
+			FPyWrapperEnumIteratorPtr NewIter = FPyWrapperEnumIteratorPtr::StealReference(FPyWrapperEnumIterator::New(&PyWrapperEnumIteratorType));
+			if (FPyWrapperEnumIterator::Init(NewIter, EnumEntriesArray) != 0)
+			{
+				return nullptr;
+			}
+			
+			return (PyObject*)NewIter.Release();
+		}
+	};
+
+	struct FSequenceFuncs
+	{
+		static Py_ssize_t Len(PyTypeObject* InSelf)
+		{
+			if (const FPyWrapperEnumMetaData* EnumMetaData = FPyWrapperEnumMetaData::GetMetaData(InSelf))
+			{
+				return EnumMetaData->EnumEntries.Num();
+			}
+			return 0;
+		}
+	};
+
+	PyTypeObject PyType = {
+		PyVarObject_HEAD_INIT(nullptr, 0)
+		"_EnumType", /* tp_name */
+		0, /* tp_basicsize */
+	};
+
+	PyType.tp_base = &PyType_Type;
+	PyType.tp_iter = (getiterfunc)&FFuncs::GetIter;
+
+	PyType.tp_flags = Py_TPFLAGS_DEFAULT | Py_TPFLAGS_BASETYPE;
+	PyType.tp_doc = "Metaclass type for all UE4 exposed enum instances";
+
+	static PySequenceMethods PySequence;
+	PySequence.sq_length = (lenfunc)&FSequenceFuncs::Len;
+
+	PyType.tp_as_sequence = &PySequence;
+
+	return PyType;
+}
+
+PyTypeObject InitializePyWrapperEnumIteratorType()
+{
+	struct FFuncs
+	{
+		static PyObject* New(PyTypeObject* InType, PyObject* InArgs, PyObject* InKwds)
+		{
+			return (PyObject*)FPyWrapperEnumIterator::New(InType);
+		}
+
+		static void Dealloc(FPyWrapperEnumIterator* InSelf)
+		{
+			FPyWrapperEnumIterator::Free(InSelf);
+		}
+
+		static int Init(FPyWrapperEnumIterator* InSelf, PyObject* InArgs, PyObject* InKwds)
+		{
+			PyObject* PyObj = nullptr;
+			if (!PyArg_ParseTuple(InArgs, "O:call", &PyObj))
+			{
+				return -1;
+			}
+
+			if (PyObject_IsInstance(PyObj, (PyObject*)&PyWrapperEnumType) != 1)
+			{
+				PyUtil::SetPythonError(PyExc_TypeError, InSelf, *FString::Printf(TEXT("Cannot initialize '%s' with an instance of '%s'"), *PyUtil::GetFriendlyTypename(InSelf), *PyUtil::GetFriendlyTypename(PyObj)));
+				return -1;
+			}
+
+			FPyWrapperEnumArrayView EnumEntriesArray;
+			if (const FPyWrapperEnumMetaData* EnumMetaData = FPyWrapperEnumMetaData::GetMetaData((FPyWrapperEnum*)PyObj))
+			{
+				EnumEntriesArray = EnumMetaData->EnumEntries;
+			}
+
+			return FPyWrapperEnumIterator::Init(InSelf, EnumEntriesArray);
+		}
+
+		static FPyWrapperEnumIterator* GetIter(FPyWrapperEnumIterator* InSelf)
+		{
+			return FPyWrapperEnumIterator::GetIter(InSelf);
+		}
+
+		static PyObject* IterNext(FPyWrapperEnumIterator* InSelf)
+		{
+			return FPyWrapperEnumIterator::IterNext(InSelf);
+		}
+	};
+
+	PyTypeObject PyType = {
+		PyVarObject_HEAD_INIT(nullptr, 0)
+		"_EnumIterator", /* tp_name */
+		sizeof(FPyWrapperEnumIterator), /* tp_basicsize */
+	};
+
+	PyType.tp_new = (newfunc)&FFuncs::New;
+	PyType.tp_dealloc = (destructor)&FFuncs::Dealloc;
+	PyType.tp_init = (initproc)&FFuncs::Init;
+	PyType.tp_iter = (getiterfunc)&FFuncs::GetIter;
+	PyType.tp_iternext = (iternextfunc)&FFuncs::IterNext;
+
+	PyType.tp_flags = Py_TPFLAGS_DEFAULT;
+	PyType.tp_doc = "Type for all UE4 exposed enum iterators";
+
+	return PyType;
+}
+
 PyTypeObject PyWrapperEnumType = InitializePyWrapperEnumType();
 PyTypeObject PyWrapperEnumValueDescrType = InitializePyWrapperEnumValueDescrType();
+PyTypeObject PyWrapperEnumMetaclassType = InitializePyWrapperEnumMetaclassType();
+PyTypeObject PyWrapperEnumIteratorType = InitializePyWrapperEnumIteratorType();
 
 FPyWrapperEnumMetaData::FPyWrapperEnumMetaData()
 	: Enum(nullptr)
+	, bFinalized(false)
 {
 }
 
@@ -222,6 +672,11 @@ UEnum* FPyWrapperEnumMetaData::GetEnum(PyTypeObject* PyType)
 	return PyWrapperMetaData ? PyWrapperMetaData->Enum : nullptr;
 }
 
+UEnum* FPyWrapperEnumMetaData::GetEnum(FPyWrapperEnum* Instance)
+{
+	return GetEnum(Py_TYPE(Instance));
+}
+
 bool FPyWrapperEnumMetaData::IsEnumDeprecated(PyTypeObject* PyType, FString* OutDeprecationMessage)
 {
 	if (FPyWrapperEnumMetaData* PyWrapperMetaData = FPyWrapperEnumMetaData::GetMetaData(PyType))
@@ -239,6 +694,26 @@ bool FPyWrapperEnumMetaData::IsEnumDeprecated(PyTypeObject* PyType, FString* Out
 	return false;
 }
 
+bool FPyWrapperEnumMetaData::IsEnumDeprecated(FPyWrapperEnum* Instance, FString* OutDeprecationMessage)
+{
+	return IsEnumDeprecated(Py_TYPE(Instance), OutDeprecationMessage);
+}
+
+bool FPyWrapperEnumMetaData::IsEnumFinalized(PyTypeObject* PyType)
+{
+	if (FPyWrapperEnumMetaData* PyWrapperMetaData = FPyWrapperEnumMetaData::GetMetaData(PyType))
+	{
+		return PyWrapperMetaData->bFinalized;
+	}
+
+	return false;
+}
+
+bool FPyWrapperEnumMetaData::IsEnumFinalized(FPyWrapperEnum* Instance)
+{
+	return IsEnumFinalized(Py_TYPE(Instance));
+}
+
 struct FPythonGeneratedEnumUtil
 {
 	static UPythonGeneratedEnum* CreateEnum(const FString& InEnumName, UObject* InEnumOuter)
@@ -255,6 +730,7 @@ struct FPythonGeneratedEnumUtil
 
 		// Add the object meta-data to the type
 		InEnum->PyMetaData.Enum = InEnum;
+		InEnum->PyMetaData.bFinalized = true;
 		FPyWrapperEnumMetaData::SetMetaData(InPyType, &InEnum->PyMetaData);
 
 		// Map the Unreal enum to the Python type
@@ -312,7 +788,11 @@ struct FPythonGeneratedEnumUtil
 		// Replace the definitions with real descriptors
 		for (const TSharedPtr<UPythonGeneratedEnum::FEnumValueDef>& EnumValueDef : InEnum->EnumValueDefs)
 		{
-			FPyWrapperEnum::SetEnumEntryValue(InPyType, EnumValueDef->Value, TCHAR_TO_UTF8(*EnumValueDef->Name), TCHAR_TO_UTF8(*EnumValueDef->DocString));
+			FPyWrapperEnum* PyEnumEntry = FPyWrapperEnum::AddEnumEntry(InPyType, EnumValueDef->Value, TCHAR_TO_UTF8(*EnumValueDef->Name), TCHAR_TO_UTF8(*EnumValueDef->DocString));
+			if (PyEnumEntry)
+			{
+				InEnum->PyMetaData.EnumEntries.Add(PyEnumEntry);
+			}
 		}
 
 		return true;
diff --git a/Engine/Plugins/Experimental/PythonScriptPlugin/Source/PythonScriptPlugin/Private/PyWrapperEnum.h b/Engine/Plugins/Experimental/PythonScriptPlugin/Source/PythonScriptPlugin/Private/PyWrapperEnum.h
--- a/Engine/Plugins/Experimental/PythonScriptPlugin/Source/PythonScriptPlugin/Private/PyWrapperEnum.h
+++ b/Engine/Plugins/Experimental/PythonScriptPlugin/Source/PythonScriptPlugin/Private/PyWrapperEnum.h
@@ -11,17 +11,53 @@
 /** Python type for FPyWrapperEnum */
 extern PyTypeObject PyWrapperEnumType;
 
+/** Python type for FPyWrapperEnumValueDescrObject */
+extern PyTypeObject PyWrapperEnumValueDescrType;
+
 /** Initialize the PyWrapperEnum types and add them to the given Python module */
 void InitializePyWrapperEnum(PyGenUtil::FNativePythonModule& ModuleInfo);
 
-/** Type for all UE4 exposed enum instances */
+/** Type for all UE4 exposed enum instances (an instance is created for each entry in the enum, before the enum type is locked for creating new instances) */
 struct FPyWrapperEnum : public FPyWrapperBase
 {
+	/** Name of this enum entry */
+	PyObject* EntryName;
+
+	/** Value of this enum entry */
+	PyObject* EntryValue;
+
+	/** New this wrapper instance (called via tp_new for Python, or directly in C++) */
+	static FPyWrapperEnum* New(PyTypeObject* InType);
+
+	/** Free this wrapper instance (called via tp_dealloc for Python) */
+	static void Free(FPyWrapperEnum* InSelf);
+
 	/** Initialize this wrapper instance (called via tp_init for Python, or directly in C++) */
 	static int Init(FPyWrapperEnum* InSelf);
 
-	/** Set the given enum entry value on the given enum type */
-	static void SetEnumEntryValue(PyTypeObject* InType, const int64 InEnumValue, const char* InEnumValueName, const char* InEnumValueDoc);
+	/** Initialize this wrapper instance (called directly in C++) */
+	static int Init(FPyWrapperEnum* InSelf, const int64 InEnumEntryValue, const char* InEnumEntryName);
+
+	/** Deinitialize this wrapper instance (called via Init and Free to restore the instance to its New state) */
+	static void Deinit(FPyWrapperEnum* InSelf);
+
+	/** Called to validate the internal state of this wrapper instance prior to operating on it (should be called by all functions that expect to operate on an initialized type; will set an error state on failure) */
+	static bool ValidateInternalState(FPyWrapperEnum* InSelf);
+
+	/** Cast the given Python object to this wrapped type (returns a new reference) */
+	static FPyWrapperEnum* CastPyObject(PyObject* InPyObject, FPyConversionResult* OutCastResult = nullptr);
+
+	/** Cast the given Python object to this wrapped type, or attempt to convert the type into a new wrapped instance (returns a new reference) */
+	static FPyWrapperEnum* CastPyObject(PyObject* InPyObject, PyTypeObject* InType, FPyConversionResult* OutCastResult = nullptr);
+
+	/** Get the name of the enum entry as a string */
+	static FString GetEnumEntryName(FPyWrapperEnum* InSelf);
+
+	/** Get the value of the enum entry as an int */
+	static int64 GetEnumEntryValue(FPyWrapperEnum* InSelf);
+
+	/** Add the given enum entry on the given enum type (returns borrowed reference) */
+	static FPyWrapperEnum* AddEnumEntry(PyTypeObject* InType, const int64 InEnumEntryValue, const char* InEnumEntryName, const char* InEnumEntryDoc);
 };
 
 /** Meta-data for all UE4 exposed enum types */
@@ -34,14 +70,32 @@ struct FPyWrapperEnumMetaData : public FPyWrapperBaseMetaData
 	/** Get the UEnum from the given type */
 	static UEnum* GetEnum(PyTypeObject* PyType);
 
+	/** Get the UEnum from the type of the given instance */
+	static UEnum* GetEnum(FPyWrapperEnum* Instance);
+
 	/** Check to see if the enum is deprecated, and optionally return its deprecation message */
 	static bool IsEnumDeprecated(PyTypeObject* PyType, FString* OutDeprecationMessage = nullptr);
 
+	/** Check to see if the enum is deprecated, and optionally return its deprecation message */
+	static bool IsEnumDeprecated(FPyWrapperEnum* Instance, FString* OutDeprecationMessage = nullptr);
+
+	/** Check to see if the enum is finalized */
+	static bool IsEnumFinalized(PyTypeObject* PyType);
+
+	/** Check to see if the enum is finalized */
+	static bool IsEnumFinalized(FPyWrapperEnum* Instance);
+
 	/** Unreal enum */
 	UEnum* Enum;
 
+	/** True if this enum type has been finalized after having all of its entries added to it */
+	bool bFinalized;
+
 	/** Set if this struct is deprecated and using it should emit a deprecation warning */
 	TOptional<FString> DeprecationMessage;
+
+	/** Array of enum entries in the order they were added (enum entries are stored as borrowed references) */
+	TArray<FPyWrapperEnum*> EnumEntries;
 };
 
 typedef TPyPtr<FPyWrapperEnum> FPyWrapperEnumPtr;
diff --git a/Engine/Plugins/Experimental/PythonScriptPlugin/Source/PythonScriptPlugin/Private/PyWrapperTypeRegistry.cpp b/Engine/Plugins/Experimental/PythonScriptPlugin/Source/PythonScriptPlugin/Private/PyWrapperTypeRegistry.cpp
--- a/Engine/Plugins/Experimental/PythonScriptPlugin/Source/PythonScriptPlugin/Private/PyWrapperTypeRegistry.cpp
+++ b/Engine/Plugins/Experimental/PythonScriptPlugin/Source/PythonScriptPlugin/Private/PyWrapperTypeRegistry.cpp
@@ -2104,13 +2104,14 @@ void FPyWrapperTypeRegistry::GenerateStubCodeForWrappedTypes() const
 
 void FPyWrapperTypeRegistry::GenerateStubCodeForWrappedType(PyTypeObject* PyType, const PyGenUtil::FGeneratedWrappedType* GeneratedTypeData, FPyFileWriter& OutPythonScript, FPyOnlineDocsSection* OutOnlineDocsSection)
 {
-	OutPythonScript.WriteLine(FString::Printf(TEXT("class %s(%s):"), UTF8_TO_TCHAR(PyType->tp_name), UTF8_TO_TCHAR(PyType->tp_base->tp_name)));
+	const FString PyTypeName = UTF8_TO_TCHAR(PyType->tp_name);
+	OutPythonScript.WriteLine(FString::Printf(TEXT("class %s(%s):"), *PyTypeName, UTF8_TO_TCHAR(PyType->tp_base->tp_name)));
 	OutPythonScript.IncreaseIndent();
 	OutPythonScript.WriteDocString(UTF8_TO_TCHAR(PyType->tp_doc));
 
 	if (OutOnlineDocsSection)
 	{
-		OutOnlineDocsSection->AccumulateClass(UTF8_TO_TCHAR(PyType->tp_name));
+		OutOnlineDocsSection->AccumulateClass(*PyTypeName);
 	}
 
 	auto GetFunctionReturnValue = [](const void* InBaseParamsAddr, const TArray<PyGenUtil::FGeneratedWrappedMethodParameter>& InOutputParams) -> FString
@@ -2332,13 +2333,8 @@ void FPyWrapperTypeRegistry::GenerateStubCodeForWrappedType(PyTypeObject* PyType
 
 			if (MetaGuid == FPyWrapperObjectMetaData::StaticTypeId())
 			{
+				// Skip the __init__ function on derived object types as the base one is already correct
 				bWriteDefaultInit = false;
-				bHasExportedClassData = true;
-
-				OutPythonScript.WriteLine(TEXT("def __init__(self, outer=None, name=\"None\"):"));
-				OutPythonScript.IncreaseIndent();
-				OutPythonScript.WriteLine(TEXT("pass"));
-				OutPythonScript.DecreaseIndent();
 			}
 			else if (MetaGuid == FPyWrapperStructMetaData::StaticTypeId())
 			{
@@ -2394,11 +2390,55 @@ void FPyWrapperTypeRegistry::GenerateStubCodeForWrappedType(PyTypeObject* PyType
 			}
 			else if (MetaGuid == FPyWrapperEnumMetaData::StaticTypeId())
 			{
-				// Enums cannot be instanced, they don't have an __init__ function exposed
+				// Skip the __init__ function on derived enums
 				bWriteDefaultInit = false;
 			}
 			// todo: have correct __init__ signatures for the other intrinsic types?
 		}
+		else if (PyType == &PyWrapperObjectType)
+		{
+			bWriteDefaultInit = false;
+			bHasExportedClassData = true;
+
+			OutPythonScript.WriteLine(TEXT("def __init__(self, outer=None, name=\"None\"):"));
+			OutPythonScript.IncreaseIndent();
+			OutPythonScript.WriteLine(TEXT("pass"));
+			OutPythonScript.DecreaseIndent();
+		}
+		else if (PyType == &PyWrapperEnumType)
+		{
+			// Enums don't really have an __init__ function at runtime, so just give them a default one (with no arguments)
+			bWriteDefaultInit = false;
+
+			OutPythonScript.WriteLine(TEXT("def __init__(self):"));
+			OutPythonScript.IncreaseIndent();
+			OutPythonScript.WriteLine(TEXT("pass"));
+			OutPythonScript.DecreaseIndent();
+		}
+		else if (PyType == &PyWrapperEnumValueDescrType)
+		{
+			bWriteDefaultInit = false;
+			bHasExportedClassData = true;
+
+			// This is a special internal decorator type used to define enum entries, which is why it has __get__ as well as __init__
+			OutPythonScript.WriteLine(TEXT("def __init__(self, enum, name, value):"));
+			OutPythonScript.IncreaseIndent();
+			OutPythonScript.WriteLine(TEXT("self.enum = enum"));
+			OutPythonScript.WriteLine(TEXT("self.name = name"));
+			OutPythonScript.WriteLine(TEXT("self.value = value"));
+			OutPythonScript.DecreaseIndent();
+
+			OutPythonScript.WriteLine(TEXT("def __get__(self, obj, type=None):"));
+			OutPythonScript.IncreaseIndent();
+			OutPythonScript.WriteLine(TEXT("return self"));
+			OutPythonScript.DecreaseIndent();
+
+			// It also needs a __repr__ function for Sphinx to generate docs correctly
+			OutPythonScript.WriteLine(TEXT("def __repr__(self):"));
+			OutPythonScript.IncreaseIndent();
+			OutPythonScript.WriteLine(TEXT("return \"{0}.{1}\".format(self.enum, self.name)"));
+			OutPythonScript.DecreaseIndent();
+		}
 
 		if (bWriteDefaultInit)
 		{
@@ -2544,25 +2584,31 @@ void FPyWrapperTypeRegistry::GenerateStubCodeForWrappedType(PyTypeObject* PyType
 		}
 		else if (MetaGuid == FPyWrapperEnumMetaData::StaticTypeId())
 		{
-			// Export enum values
-			// Also see https://www.python.org/dev/peps/pep-0435/
-
+			// Export enum entries
 			const PyGenUtil::FGeneratedWrappedEnumType* EnumType = static_cast<const PyGenUtil::FGeneratedWrappedEnumType*>(GeneratedTypeData);
 
-			bool bFirstEnumMember = true;
-
-			for (const PyGenUtil::FGeneratedWrappedEnumEntry& EnumMember : EnumType->EnumEntries)
+			if (EnumType->EnumEntries.Num() > 0)
 			{
-				bHasExportedClassData = true;
+				// Add extra line break for first enum member
+				OutPythonScript.WriteNewLine();
 
-				if (bFirstEnumMember)
+				for (const PyGenUtil::FGeneratedWrappedEnumEntry& EnumMember : EnumType->EnumEntries)
 				{
-					// Add extra line break for first enum member
-					OutPythonScript.WriteNewLine();
-					bFirstEnumMember = false;
-				}
+					const FString EntryName = UTF8_TO_TCHAR(EnumMember.EntryName.GetData());
+					const FString EntryValue = LexToString(EnumMember.EntryValue);
 
-				ExportConstantValue(UTF8_TO_TCHAR(EnumMember.EntryName.GetData()), UTF8_TO_TCHAR(EnumMember.EntryDoc.GetData()), *LexToString(EnumMember.EntryValue));
+					FString EntryDoc = UTF8_TO_TCHAR(EnumMember.EntryDoc.GetData());
+					if (EntryDoc.IsEmpty())
+					{
+						EntryDoc = EntryValue;
+					}
+					else
+					{
+						EntryDoc.InsertAt(0, *FString::Printf(TEXT("%s: "), *EntryValue));
+					}
+
+					ExportConstantValue(*EntryName, *EntryDoc, *FString::Printf(TEXT("_EnumEntry(\"%s\", \"%s\", %s)"), *PyTypeName, *EntryName, *EntryValue));
+				}
 			}
 		}
 	}
