Implemented negative indexing for Python wrapped arrays

diff --git a/Engine/Plugins/Experimental/PythonScriptPlugin/Source/PythonScriptPlugin/Private/PyUtil.cpp b/Engine/Plugins/Experimental/PythonScriptPlugin/Source/PythonScriptPlugin/Private/PyUtil.cpp
--- a/Engine/Plugins/Experimental/PythonScriptPlugin/Source/PythonScriptPlugin/Private/PyUtil.cpp
+++ b/Engine/Plugins/Experimental/PythonScriptPlugin/Source/PythonScriptPlugin/Private/PyUtil.cpp
@@ -637,6 +637,11 @@ int ValidateContainerIndexParam(const Py_ssize_t InIndex, const Py_ssize_t InLen
 	return 0;
 }
 
+Py_ssize_t ResolveContainerIndexParam(const Py_ssize_t InIndex, const Py_ssize_t InLen)
+{
+	return InIndex < 0 ? InIndex + InLen : InIndex;
+}
+
 UObject* GetOwnerObject(PyObject* InPyObj)
 {
 	FPyWrapperOwnerContext OwnerContext = FPyWrapperOwnerContext(InPyObj);
diff --git a/Engine/Plugins/Experimental/PythonScriptPlugin/Source/PythonScriptPlugin/Private/PyUtil.h b/Engine/Plugins/Experimental/PythonScriptPlugin/Source/PythonScriptPlugin/Private/PyUtil.h
--- a/Engine/Plugins/Experimental/PythonScriptPlugin/Source/PythonScriptPlugin/Private/PyUtil.h
+++ b/Engine/Plugins/Experimental/PythonScriptPlugin/Source/PythonScriptPlugin/Private/PyUtil.h
@@ -215,6 +215,9 @@ namespace PyUtil
 	/** Validate that the given index is valid for the container length */
 	int ValidateContainerIndexParam(const Py_ssize_t InIndex, const Py_ssize_t InLen, const UProperty* InProp, const TCHAR* InErrorCtxt);
 
+	/** Resolve a container index (taking into account negative indices) */
+	Py_ssize_t ResolveContainerIndexParam(const Py_ssize_t InIndex, const Py_ssize_t InLen);
+
 	/**
 	 * Given a Python object, try and get the owner Unreal object for the instance.
 	 * For wrapped objects this is the wrapped instance, for wrapped structs it will attempt to walk through the owner chain to find a wrapped object.
diff --git a/Engine/Plugins/Experimental/PythonScriptPlugin/Source/PythonScriptPlugin/Private/PyWrapperArray.cpp b/Engine/Plugins/Experimental/PythonScriptPlugin/Source/PythonScriptPlugin/Private/PyWrapperArray.cpp
--- a/Engine/Plugins/Experimental/PythonScriptPlugin/Source/PythonScriptPlugin/Private/PyWrapperArray.cpp
+++ b/Engine/Plugins/Experimental/PythonScriptPlugin/Source/PythonScriptPlugin/Private/PyWrapperArray.cpp
@@ -407,16 +407,17 @@ PyObject* FPyWrapperArray::GetItem(FPyWrapperArray* InSelf, Py_ssize_t InIndex)
 
 	FScriptArrayHelper SelfScriptArrayHelper(InSelf->ArrayProp, InSelf->ArrayInstance);
 	const int32 ElementCount = SelfScriptArrayHelper.Num();
+	const Py_ssize_t ResolvedIndex = PyUtil::ResolveContainerIndexParam(InIndex, ElementCount);
 
-	if (PyUtil::ValidateContainerIndexParam(InIndex, ElementCount, InSelf->ArrayProp, *PyUtil::GetErrorContext(InSelf)) != 0)
+	if (PyUtil::ValidateContainerIndexParam(ResolvedIndex, ElementCount, InSelf->ArrayProp, *PyUtil::GetErrorContext(InSelf)) != 0)
 	{
 		return nullptr;
 	}
 
 	PyObject* PyItemObj = nullptr;
-	if (!PyConversion::PythonizeProperty(InSelf->ArrayProp->Inner, SelfScriptArrayHelper.GetRawPtr(InIndex), PyItemObj))
+	if (!PyConversion::PythonizeProperty(InSelf->ArrayProp->Inner, SelfScriptArrayHelper.GetRawPtr(ResolvedIndex), PyItemObj))
 	{
-		PyUtil::SetPythonError(PyExc_TypeError, InSelf, *FString::Printf(TEXT("Failed to convert element property '%s' (%s) at index %d"), *InSelf->ArrayProp->Inner->GetName(), *InSelf->ArrayProp->Inner->GetClass()->GetName(), InIndex));
+		PyUtil::SetPythonError(PyExc_TypeError, InSelf, *FString::Printf(TEXT("Failed to convert element property '%s' (%s) at index %d"), *InSelf->ArrayProp->Inner->GetName(), *InSelf->ArrayProp->Inner->GetClass()->GetName(), ResolvedIndex));
 		return nullptr;
 	}
 	return PyItemObj;
@@ -431,16 +432,17 @@ int FPyWrapperArray::SetItem(FPyWrapperArray* InSelf, Py_ssize_t InIndex, PyObje
 
 	FScriptArrayHelper SelfScriptArrayHelper(InSelf->ArrayProp, InSelf->ArrayInstance);
 	const int32 ElementCount = SelfScriptArrayHelper.Num();
+	const Py_ssize_t ResolvedIndex = PyUtil::ResolveContainerIndexParam(InIndex, ElementCount);
 
-	const int ValidateIndexResult = PyUtil::ValidateContainerIndexParam(InIndex, ElementCount, InSelf->ArrayProp, *PyUtil::GetErrorContext(InSelf));
+	const int ValidateIndexResult = PyUtil::ValidateContainerIndexParam(ResolvedIndex, ElementCount, InSelf->ArrayProp, *PyUtil::GetErrorContext(InSelf));
 	if (ValidateIndexResult != 0)
 	{
 		return ValidateIndexResult;
 	}
 
-	if (!PyConversion::NativizeProperty(InValue, InSelf->ArrayProp->Inner, SelfScriptArrayHelper.GetRawPtr(InIndex)))
+	if (!PyConversion::NativizeProperty(InValue, InSelf->ArrayProp->Inner, SelfScriptArrayHelper.GetRawPtr(ResolvedIndex)))
 	{
-		PyUtil::SetPythonError(PyExc_TypeError, InSelf, *FString::Printf(TEXT("Failed to convert element property '%s' (%s) at index %d"), *InSelf->ArrayProp->Inner->GetName(), *InSelf->ArrayProp->Inner->GetClass()->GetName(), InIndex));
+		PyUtil::SetPythonError(PyExc_TypeError, InSelf, *FString::Printf(TEXT("Failed to convert element property '%s' (%s) at index %d"), *InSelf->ArrayProp->Inner->GetName(), *InSelf->ArrayProp->Inner->GetClass()->GetName(), ResolvedIndex));
 		return -1;
 	}
 
@@ -640,9 +642,11 @@ Py_ssize_t FPyWrapperArray::Index(FPyWrapperArray* InSelf, PyObject* InValue, Py
 
 	FScriptArrayHelper SelfScriptArrayHelper(InSelf->ArrayProp, InSelf->ArrayInstance);
 	const int32 ElementCount = SelfScriptArrayHelper.Num();
+	const Py_ssize_t ResolvedStartIndex = PyUtil::ResolveContainerIndexParam(InStartIndex, ElementCount);
+	const Py_ssize_t ResolvedStopIndex = PyUtil::ResolveContainerIndexParam(InStopIndex, ElementCount);
 
-	const int32 StartIndex = FMath::Min((int32)InStartIndex, ElementCount);
-	const int32 StopIndex = FMath::Max((int32)InStopIndex, ElementCount);
+	const int32 StartIndex = FMath::Min((int32)ResolvedStartIndex, ElementCount);
+	const int32 StopIndex = FMath::Max((int32)ResolvedStopIndex, ElementCount);
 
 	int32 ReturnIndex = INDEX_NONE;
 	for (int32 ElementIndex = StartIndex; ElementIndex < StopIndex; ++ElementIndex)
@@ -672,8 +676,9 @@ int FPyWrapperArray::Insert(FPyWrapperArray* InSelf, Py_ssize_t InIndex, PyObjec
 
 	FScriptArrayHelper SelfScriptArrayHelper(InSelf->ArrayProp, InSelf->ArrayInstance);
 	const int32 ElementCount = SelfScriptArrayHelper.Num();
+	const Py_ssize_t ResolvedIndex = PyUtil::ResolveContainerIndexParam(InIndex, ElementCount);
 
-	const int32 InsertIndex = FMath::Min((int32)InIndex, ElementCount);
+	const int32 InsertIndex = FMath::Min((int32)ResolvedIndex, ElementCount);
 	SelfScriptArrayHelper.InsertValues(InsertIndex);
 
 	if (!PyConversion::NativizeProperty(InValue, InSelf->ArrayProp->Inner, SelfScriptArrayHelper.GetRawPtr(InsertIndex)))
@@ -695,21 +700,21 @@ PyObject* FPyWrapperArray::Pop(FPyWrapperArray* InSelf, Py_ssize_t InIndex)
 
 	FScriptArrayHelper SelfScriptArrayHelper(InSelf->ArrayProp, InSelf->ArrayInstance);
 	const int32 ElementCount = SelfScriptArrayHelper.Num();
+	const Py_ssize_t ResolvedIndex = PyUtil::ResolveContainerIndexParam(InIndex, ElementCount);
 
-	int32 ValueIndex = InIndex == -1 ? ElementCount - 1 : (int32)InIndex;
-	if (PyUtil::ValidateContainerIndexParam(ValueIndex, ElementCount, InSelf->ArrayProp, *PyUtil::GetErrorContext(InSelf)) != 0)
+	if (PyUtil::ValidateContainerIndexParam(ResolvedIndex, ElementCount, InSelf->ArrayProp, *PyUtil::GetErrorContext(InSelf)) != 0)
 	{
 		return nullptr;
 	}
 
 	PyObject* PyReturnValue = nullptr;
-	if (!PyConversion::PythonizeProperty(InSelf->ArrayProp->Inner, SelfScriptArrayHelper.GetRawPtr(ValueIndex), PyReturnValue))
+	if (!PyConversion::PythonizeProperty(InSelf->ArrayProp->Inner, SelfScriptArrayHelper.GetRawPtr(ResolvedIndex), PyReturnValue))
 	{
-		PyUtil::SetPythonError(PyExc_TypeError, InSelf, *FString::Printf(TEXT("Failed to convert element property '%s' (%s) at index %d"), *InSelf->ArrayProp->Inner->GetName(), *InSelf->ArrayProp->Inner->GetClass()->GetName(), ValueIndex));
+		PyUtil::SetPythonError(PyExc_TypeError, InSelf, *FString::Printf(TEXT("Failed to convert element property '%s' (%s) at index %d"), *InSelf->ArrayProp->Inner->GetName(), *InSelf->ArrayProp->Inner->GetClass()->GetName(), ResolvedIndex));
 		return nullptr;
 	}
 
-	SelfScriptArrayHelper.RemoveValues(ValueIndex);
+	SelfScriptArrayHelper.RemoveValues(ResolvedIndex);
 
 	return PyReturnValue;
 }
@@ -991,7 +996,12 @@ PyTypeObject InitializePyWrapperArrayType()
 				return nullptr;
 			}
 
-			const Py_ssize_t SliceLen = FMath::Max(InSliceStop - InSliceStart, (Py_ssize_t)0);
+			FScriptArrayHelper SelfScriptArrayHelper(InSelf->ArrayProp, InSelf->ArrayInstance);
+			const int32 SelfElementCount = SelfScriptArrayHelper.Num();
+			const Py_ssize_t ResolvedSliceStart = PyUtil::ResolveContainerIndexParam(InSliceStart, SelfElementCount);
+			const Py_ssize_t ResolvedSliceStop = PyUtil::ResolveContainerIndexParam(InSliceStop, SelfElementCount);
+
+			const Py_ssize_t SliceLen = FMath::Max(ResolvedSliceStop - ResolvedSliceStart, (Py_ssize_t)0);
 
 			const PyUtil::FPropertyDef SelfElementDef = InSelf->ArrayProp->Inner;
 			FPyWrapperArrayPtr NewArray = FPyWrapperArrayPtr::StealReference(FPyWrapperArray::New(Py_TYPE(InSelf)));
@@ -1000,13 +1010,12 @@ PyTypeObject InitializePyWrapperArrayType()
 				return nullptr;
 			}
 
-			FScriptArrayHelper SelfScriptArrayHelper(InSelf->ArrayProp, InSelf->ArrayInstance);
 			FScriptArrayHelper NewScriptArrayHelper(NewArray->ArrayProp, NewArray->ArrayInstance);
 			NewScriptArrayHelper.Resize(SliceLen);
 
-			for (Py_ssize_t ElementIndex = InSliceStart; ElementIndex < InSliceStop; ++ElementIndex)
+			for (Py_ssize_t ElementIndex = ResolvedSliceStart; ElementIndex < ResolvedSliceStop; ++ElementIndex)
 			{
-				InSelf->ArrayProp->Inner->CopyCompleteValue(NewScriptArrayHelper.GetRawPtr((int32)(ElementIndex - InSliceStart)), SelfScriptArrayHelper.GetRawPtr((int32)ElementIndex));
+				InSelf->ArrayProp->Inner->CopyCompleteValue(NewScriptArrayHelper.GetRawPtr((int32)(ElementIndex - ResolvedSliceStart)), SelfScriptArrayHelper.GetRawPtr((int32)ElementIndex));
 			}
 			return (PyObject*)NewArray.Release();
 		}
@@ -1018,7 +1027,12 @@ PyTypeObject InitializePyWrapperArrayType()
 				return -1;
 			}
 
-			const Py_ssize_t SliceLen = FMath::Max(InSliceStop - InSliceStart, (Py_ssize_t)0);
+			FScriptArrayHelper SelfScriptArrayHelper(InSelf->ArrayProp, InSelf->ArrayInstance);
+			const int32 SelfElementCount = SelfScriptArrayHelper.Num();
+			const Py_ssize_t ResolvedSliceStart = PyUtil::ResolveContainerIndexParam(InSliceStart, SelfElementCount);
+			const Py_ssize_t ResolvedSliceStop = PyUtil::ResolveContainerIndexParam(InSliceStop, SelfElementCount);
+
+			const Py_ssize_t SliceLen = FMath::Max(ResolvedSliceStop - ResolvedSliceStart, (Py_ssize_t)0);
 
 			// Value will be null when performing a slice delete
 			FPyWrapperArrayPtr Value;
@@ -1033,18 +1047,17 @@ PyTypeObject InitializePyWrapperArrayType()
 				}
 			}
 
-			FScriptArrayHelper SelfScriptArrayHelper(InSelf->ArrayProp, InSelf->ArrayInstance);
-			SelfScriptArrayHelper.RemoveValues((int32)InSliceStart, (int32)SliceLen);
+			SelfScriptArrayHelper.RemoveValues((int32)ResolvedSliceStart, (int32)SliceLen);
 
 			if (Value)
 			{
 				FScriptArrayHelper ValueScriptArrayHelper(Value->ArrayProp, Value->ArrayInstance);
 				const int32 ValueElementCount = ValueScriptArrayHelper.Num();
 
-				SelfScriptArrayHelper.InsertValues((int32)InSliceStart, ValueElementCount);
+				SelfScriptArrayHelper.InsertValues((int32)ResolvedSliceStart, ValueElementCount);
 				for (int32 ElementIndex = 0; ElementIndex < ValueElementCount; ++ElementIndex)
 				{
-					InSelf->ArrayProp->Inner->CopyCompleteValue(SelfScriptArrayHelper.GetRawPtr(InSliceStart + ElementIndex), ValueScriptArrayHelper.GetRawPtr(ElementIndex));
+					InSelf->ArrayProp->Inner->CopyCompleteValue(SelfScriptArrayHelper.GetRawPtr(ResolvedSliceStart + ElementIndex), ValueScriptArrayHelper.GetRawPtr(ElementIndex));
 				}
 			}
 
diff --git a/Engine/Plugins/Experimental/PythonScriptPlugin/Source/PythonScriptPlugin/Private/PyWrapperFixedArray.cpp b/Engine/Plugins/Experimental/PythonScriptPlugin/Source/PythonScriptPlugin/Private/PyWrapperFixedArray.cpp
--- a/Engine/Plugins/Experimental/PythonScriptPlugin/Source/PythonScriptPlugin/Private/PyWrapperFixedArray.cpp
+++ b/Engine/Plugins/Experimental/PythonScriptPlugin/Source/PythonScriptPlugin/Private/PyWrapperFixedArray.cpp
@@ -386,15 +386,16 @@ PyObject* FPyWrapperFixedArray::GetItem(FPyWrapperFixedArray* InSelf, Py_ssize_t
 		return nullptr;
 	}
 
-	if (PyUtil::ValidateContainerIndexParam(InIndex, InSelf->ArrayProp->ArrayDim, InSelf->ArrayProp, *PyUtil::GetErrorContext(InSelf)) != 0)
+	const Py_ssize_t ResolvedIndex = PyUtil::ResolveContainerIndexParam(InIndex, InSelf->ArrayProp->ArrayDim);
+	if (PyUtil::ValidateContainerIndexParam(ResolvedIndex, InSelf->ArrayProp->ArrayDim, InSelf->ArrayProp, *PyUtil::GetErrorContext(InSelf)) != 0)
 	{
 		return nullptr;
 	}
 
 	PyObject* PyItemObj = nullptr;
-	if (!PyConversion::PythonizeProperty_Direct(InSelf->ArrayProp, GetItemPtr(InSelf, InIndex), PyItemObj))
+	if (!PyConversion::PythonizeProperty_Direct(InSelf->ArrayProp, GetItemPtr(InSelf, ResolvedIndex), PyItemObj))
 	{
-		PyUtil::SetPythonError(PyExc_TypeError, InSelf, *FString::Printf(TEXT("Failed to convert property '%s' (%s) at index %d"), *InSelf->ArrayProp->GetName(), *InSelf->ArrayProp->GetClass()->GetName(), InIndex));
+		PyUtil::SetPythonError(PyExc_TypeError, InSelf, *FString::Printf(TEXT("Failed to convert property '%s' (%s) at index %d"), *InSelf->ArrayProp->GetName(), *InSelf->ArrayProp->GetClass()->GetName(), ResolvedIndex));
 		return nullptr;
 	}
 	return PyItemObj;
@@ -407,15 +408,16 @@ int FPyWrapperFixedArray::SetItem(FPyWrapperFixedArray* InSelf, Py_ssize_t InInd
 		return -1;
 	}
 
-	const int ValidateIndexResult = PyUtil::ValidateContainerIndexParam(InIndex, InSelf->ArrayProp->ArrayDim, InSelf->ArrayProp, *PyUtil::GetErrorContext(InSelf));
+	const Py_ssize_t ResolvedIndex = PyUtil::ResolveContainerIndexParam(InIndex, InSelf->ArrayProp->ArrayDim);
+	const int ValidateIndexResult = PyUtil::ValidateContainerIndexParam(ResolvedIndex, InSelf->ArrayProp->ArrayDim, InSelf->ArrayProp, *PyUtil::GetErrorContext(InSelf));
 	if (ValidateIndexResult != 0)
 	{
 		return ValidateIndexResult;
 	}
 
-	if (!PyConversion::NativizeProperty_Direct(InValue, InSelf->ArrayProp, GetItemPtr(InSelf, InIndex)))
+	if (!PyConversion::NativizeProperty_Direct(InValue, InSelf->ArrayProp, GetItemPtr(InSelf, ResolvedIndex)))
 	{
-		PyUtil::SetPythonError(PyExc_TypeError, InSelf, *FString::Printf(TEXT("Failed to convert property '%s' (%s) at index %d"), *InSelf->ArrayProp->GetName(), *InSelf->ArrayProp->GetClass()->GetName(), InIndex));
+		PyUtil::SetPythonError(PyExc_TypeError, InSelf, *FString::Printf(TEXT("Failed to convert property '%s' (%s) at index %d"), *InSelf->ArrayProp->GetName(), *InSelf->ArrayProp->GetClass()->GetName(), ResolvedIndex));
 		return -1;
 	}
 
@@ -688,7 +690,9 @@ PyTypeObject InitializePyWrapperFixedArrayType()
 				return nullptr;
 			}
 
-			const Py_ssize_t SliceLen = InSliceStop - InSliceStart;
+			const Py_ssize_t ResolvedSliceStart = PyUtil::ResolveContainerIndexParam(InSliceStart, InSelf->ArrayProp->ArrayDim);
+			const Py_ssize_t ResolvedSliceStop = PyUtil::ResolveContainerIndexParam(InSliceStop, InSelf->ArrayProp->ArrayDim);
+			const Py_ssize_t SliceLen = ResolvedSliceStop - ResolvedSliceStart;
 			if (SliceLen <= 0)
 			{
 				PyUtil::SetPythonError(PyExc_Exception, InSelf, TEXT("Slice length must be greater than zero"));
@@ -702,9 +706,9 @@ PyTypeObject InitializePyWrapperFixedArrayType()
 				return nullptr;
 			}
 
-			for (Py_ssize_t ArrIndex = InSliceStart; ArrIndex < InSliceStop; ++ArrIndex)
+			for (Py_ssize_t ArrIndex = ResolvedSliceStart; ArrIndex < ResolvedSliceStop; ++ArrIndex)
 			{
-				InSelf->ArrayProp->CopySingleValue(FPyWrapperFixedArray::GetItemPtr(NewArray, ArrIndex - InSliceStart), FPyWrapperFixedArray::GetItemPtr(InSelf, ArrIndex));
+				InSelf->ArrayProp->CopySingleValue(FPyWrapperFixedArray::GetItemPtr(NewArray, ArrIndex - ResolvedSliceStart), FPyWrapperFixedArray::GetItemPtr(InSelf, ArrIndex));
 			}
 			return (PyObject*)NewArray.Release();
 		}
diff --git a/Engine/Plugins/Experimental/PythonScriptPlugin/Source/PythonScriptPlugin/Private/PyWrapperFixedArray.h b/Engine/Plugins/Experimental/PythonScriptPlugin/Source/PythonScriptPlugin/Private/PyWrapperFixedArray.h
--- a/Engine/Plugins/Experimental/PythonScriptPlugin/Source/PythonScriptPlugin/Private/PyWrapperFixedArray.h
+++ b/Engine/Plugins/Experimental/PythonScriptPlugin/Source/PythonScriptPlugin/Private/PyWrapperFixedArray.h
@@ -50,7 +50,7 @@ struct FPyWrapperFixedArray : public FPyWrapperBase
 	/** Cast the given Python object to this wrapped type, or attempt to convert the type into a new wrapped instance (returns a new reference) */
 	static FPyWrapperFixedArray* CastPyObject(PyObject* InPyObject, PyTypeObject* InType, const PyUtil::FPropertyDef& InPropDef, FPyConversionResult* OutCastResult = nullptr);
 
-	/** Get the raw pointer to the element at index N */
+	/** Get the raw pointer to the element at index N (negative indexing not supported) */
 	static void* GetItemPtr(FPyWrapperFixedArray* InSelf, Py_ssize_t InIndex);
 
 	/** Get the length of this container (equivalent to 'len(x)' in Python) */
