Fixed crash on shutdown after binding python member function to delegate

diff --git a/Engine/Plugins/Experimental/PythonScriptPlugin/Source/PythonScriptPlugin/Private/PyWrapperDelegate.cpp b/Engine/Plugins/Experimental/PythonScriptPlugin/Source/PythonScriptPlugin/Private/PyWrapperDelegate.cpp
--- a/Engine/Plugins/Experimental/PythonScriptPlugin/Source/PythonScriptPlugin/Private/PyWrapperDelegate.cpp
+++ b/Engine/Plugins/Experimental/PythonScriptPlugin/Source/PythonScriptPlugin/Private/PyWrapperDelegate.cpp
@@ -24,6 +24,11 @@ void UPythonCallableForDelegate::BeginDestroy()
 		FPyScopedGIL GIL;
 		PyCallable.Reset();
 	}
+	else
+	{
+		// Release ownership if Python has been shut down to avoid attempting to delete the callable (which is already dead)
+		PyCallable.Release();
+	}
 #endif	// WITH_PYTHON
 
 	Super::BeginDestroy();
