Fixed Python menus failing to execute scripts with spaces in their name or path

diff --git a/Engine/Plugins/Experimental/PythonScriptPlugin/Source/PythonScriptPlugin/Private/PythonScriptPlugin.cpp b/Engine/Plugins/Experimental/PythonScriptPlugin/Source/PythonScriptPlugin/Private/PythonScriptPlugin.cpp
--- a/Engine/Plugins/Experimental/PythonScriptPlugin/Source/PythonScriptPlugin/Private/PythonScriptPlugin.cpp
+++ b/Engine/Plugins/Experimental/PythonScriptPlugin/Source/PythonScriptPlugin/Private/PythonScriptPlugin.cpp
@@ -275,7 +275,7 @@ private:
 		);
 		MenuBuilder.AddSubMenu(
 			LOCTEXT("RecentPythonsSubMenu", "Recent Python Scripts"),
-			LOCTEXT("RecentPythonsSubMenu_ToolTip", "Select a level to load"),
+			LOCTEXT("RecentPythonsSubMenu_ToolTip", "Select a recent Python Script file and Execute it."),
 			FNewMenuDelegate::CreateRaw(this, &FPythonCommandMenuImpl::MakeRecentPythonScriptMenu),
 			false,
 			FSlateIcon(FEditorStyle::GetStyleSetName(), "MainFrame.RecentLevels")
@@ -288,7 +288,7 @@ private:
 		if (RecentsFiles.IsValidIndex(Index))
 		{
 			FString PyCopied = RecentsFiles[Index];
-			GEngine->Exec(NULL, *FString::Printf(TEXT("py %s"), *PyCopied));
+			GEngine->Exec(NULL, *FString::Printf(TEXT("py \"%s\""), *PyCopied));
 		}
 	}
 
@@ -318,7 +318,7 @@ private:
 					bRecentsFilesDirty = true;
 				}
 
-				GEngine->Exec(NULL, *FString::Printf(TEXT("py %s"), *OpenedFiles.Last()));
+				GEngine->Exec(NULL, *FString::Printf(TEXT("py \"%s\""), *OpenedFiles.Last()));
 			}
 		}
 	}
