Added missing build dependency

diff --git a/Engine/Plugins/Experimental/PythonScriptPlugin/Source/PythonScriptPlugin/PythonScriptPlugin.Build.cs b/Engine/Plugins/Experimental/PythonScriptPlugin/Source/PythonScriptPlugin/PythonScriptPlugin.Build.cs
--- a/Engine/Plugins/Experimental/PythonScriptPlugin/Source/PythonScriptPlugin/PythonScriptPlugin.Build.cs
+++ b/Engine/Plugins/Experimental/PythonScriptPlugin/Source/PythonScriptPlugin/PythonScriptPlugin.Build.cs
@@ -22,6 +22,7 @@ namespace UnrealBuildTool.Rules
 					"Python",
 					"Slate",
 					"SlateCore",
+					"InputCore",
 				}
 			);
 
