Python wrapped delegate functions cannot be coerced down to UFunction as there will be serialization differences in the future
#rnx

diff --git a/Engine/Plugins/Experimental/PythonScriptPlugin/Source/PythonScriptPlugin/Private/PyWrapperTypeRegistry.cpp b/Engine/Plugins/Experimental/PythonScriptPlugin/Source/PythonScriptPlugin/Private/PyWrapperTypeRegistry.cpp
--- a/Engine/Plugins/Experimental/PythonScriptPlugin/Source/PythonScriptPlugin/Private/PyWrapperTypeRegistry.cpp
+++ b/Engine/Plugins/Experimental/PythonScriptPlugin/Source/PythonScriptPlugin/Private/PyWrapperTypeRegistry.cpp
@@ -1776,7 +1776,7 @@ PyTypeObject* FPyWrapperTypeRegistry::GenerateWrappedDelegateType(const UFunctio
 	UClass* PythonCallableForDelegateClass = nullptr;
 	{
 		PythonCallableForDelegateClass = NewObject<UClass>(GetPythonTypeContainer(), *FString::Printf(TEXT("%s__PythonCallable"), *DelegateBaseTypename), RF_Public);
-		UFunction* PythonCallableForDelegateFunc = (UFunction*)StaticDuplicateObject(InDelegateSignature, PythonCallableForDelegateClass, UPythonCallableForDelegate::GeneratedFuncName, RF_AllFlags, UFunction::StaticClass());
+		UFunction* PythonCallableForDelegateFunc = (UFunction*)StaticDuplicateObject(InDelegateSignature, PythonCallableForDelegateClass, UPythonCallableForDelegate::GeneratedFuncName);
 		PythonCallableForDelegateFunc->FunctionFlags = (PythonCallableForDelegateFunc->FunctionFlags | FUNC_Native) & ~(FUNC_Delegate | FUNC_MulticastDelegate);
 		PythonCallableForDelegateFunc->SetNativeFunc(&UPythonCallableForDelegate::CallPythonNative);
 		PythonCallableForDelegateFunc->StaticLink(true);
