Added extra control over executing Python commands

You're now able to get back the execution result and log output from running a Python command, as well as control the execution mode of the command, and whether it should be run in "unattended" mode.

This is handled by the new IPythonScriptPlugin::ExecPythonCommandEx function which takes a FPythonCommandEx struct containing the input flags, mode, and command, as well as fields to fill in with the result and log output.

The mode can be set to one of the following:
  - ExecuteFile: Execute the Python command as a file. This allows you to execute either a literal Python script containing multiple statements, or a file with optional arguments.
  - ExecuteStatement: Execute the Python command as a single statement. This will execute a single statement and print the result. This mode cannot run files.
  - EvaluateStatement: Evaluate the Python command as a single statement. This will evaluate a single statement and return the result. This mode cannot run files.

ExecuteFile is the default to preserve existing behavior, however, since ExecuteStatement allows for a REPL like environment, there is now an additional command executor available ("Python (REPL)") which lets you build and test Python scripts in a similar manner to an interactive command line Python environment.

diff --git a/Engine/Plugins/Experimental/PythonScriptPlugin/Source/PythonScriptPlugin/Private/PyCore.cpp b/Engine/Plugins/Experimental/PythonScriptPlugin/Source/PythonScriptPlugin/Private/PyCore.cpp
--- a/Engine/Plugins/Experimental/PythonScriptPlugin/Source/PythonScriptPlugin/Private/PyCore.cpp
+++ b/Engine/Plugins/Experimental/PythonScriptPlugin/Source/PythonScriptPlugin/Private/PyCore.cpp
@@ -4,6 +4,7 @@
 #include "PyGenUtil.h"
 #include "PyReferenceCollector.h"
 #include "PyWrapperTypeRegistry.h"
+#include "IPythonScriptPlugin.h"
 
 #include "PyWrapperBase.h"
 #include "PyWrapperObject.h"
@@ -1009,13 +1010,35 @@ PyTypeObject PyUFunctionDefType = InitializePyUFunctionDefType();
 namespace PyCore
 {
 
+bool ShouldLog(const FString& InLogMessage)
+{
+	if (InLogMessage.IsEmpty())
+	{
+		return false;
+	}
+
+	for (const TCHAR LogChar : InLogMessage)
+	{
+		if (!FChar::IsWhitespace(LogChar))
+		{
+			return true;
+		}
+	}
+
+	return false;
+}
+
 PyObject* Log(PyObject* InSelf, PyObject* InArgs)
 {
 	PyObject* PyObj = nullptr;
 	if (PyArg_ParseTuple(InArgs, "O:log", &PyObj))
 	{
 		const FString LogMessage = PyUtil::PyObjectToUEString(PyObj);
-		UE_LOG(LogPython, Log, TEXT("%s"), *LogMessage);
+		if (ShouldLog(LogMessage))
+		{
+			UE_LOG(LogPython, Log, TEXT("%s"), *LogMessage);
+			GetPythonLogCapture().Broadcast(EPythonLogOutputType::Info, *LogMessage);
+		}
 
 		Py_RETURN_NONE;
 	}
@@ -1026,10 +1049,14 @@ PyObject* Log(PyObject* InSelf, PyObject* InArgs)
 PyObject* LogWarning(PyObject* InSelf, PyObject* InArgs)
 {
 	PyObject* PyObj = nullptr;
-	if (PyArg_ParseTuple(InArgs, "O:log", &PyObj))
+	if (PyArg_ParseTuple(InArgs, "O:log_warning", &PyObj))
 	{
 		const FString LogMessage = PyUtil::PyObjectToUEString(PyObj);
-		UE_LOG(LogPython, Warning, TEXT("%s"), *LogMessage);
+		if (ShouldLog(LogMessage))
+		{
+			UE_LOG(LogPython, Warning, TEXT("%s"), *LogMessage);
+			GetPythonLogCapture().Broadcast(EPythonLogOutputType::Warning, *LogMessage);
+		}
 
 		Py_RETURN_NONE;
 	}
@@ -1040,10 +1067,14 @@ PyObject* LogWarning(PyObject* InSelf, PyObject* InArgs)
 PyObject* LogError(PyObject* InSelf, PyObject* InArgs)
 {
 	PyObject* PyObj = nullptr;
-	if (PyArg_ParseTuple(InArgs, "O:log", &PyObj))
+	if (PyArg_ParseTuple(InArgs, "O:log_error", &PyObj))
 	{
 		const FString LogMessage = PyUtil::PyObjectToUEString(PyObj);
-		UE_LOG(LogPython, Error, TEXT("%s"), *LogMessage);
+		if (ShouldLog(LogMessage))
+		{
+			UE_LOG(LogPython, Error, TEXT("%s"), *LogMessage);
+			GetPythonLogCapture().Broadcast(EPythonLogOutputType::Error, *LogMessage);
+		}
 
 		Py_RETURN_NONE;
 	}
@@ -1717,6 +1748,12 @@ void InitializeModule()
 	FPyWrapperTypeRegistry::Get().RegisterNativePythonModule(MoveTemp(NativePythonModule));
 }
 
+FPythonLogCapture& GetPythonLogCapture()
+{
+	static FPythonLogCapture PythonLogCapture;
+	return PythonLogCapture;
+}
+
 }
 
 #endif	// WITH_PYTHON
diff --git a/Engine/Plugins/Experimental/PythonScriptPlugin/Source/PythonScriptPlugin/Private/PyCore.h b/Engine/Plugins/Experimental/PythonScriptPlugin/Source/PythonScriptPlugin/Private/PyCore.h
--- a/Engine/Plugins/Experimental/PythonScriptPlugin/Source/PythonScriptPlugin/Private/PyCore.h
+++ b/Engine/Plugins/Experimental/PythonScriptPlugin/Source/PythonScriptPlugin/Private/PyCore.h
@@ -10,9 +10,12 @@
 #include "CoreMinimal.h"
 #include "Misc/EnumClassFlags.h"
 #include "UObject/UObjectIterator.h"
+#include "Delegates/DelegateCombinations.h"
 
 #if WITH_PYTHON
 
+enum class EPythonLogOutputType : uint8;
+
 struct FSlowTask;
 
 /** Get the object that Python created transient properties should be outered to */
@@ -370,6 +373,9 @@ typedef TPyPtr<FPyUFunctionDef> FPyUFunctionDefPtr;
 namespace PyCore
 {
 	void InitializeModule();
+
+	DECLARE_MULTICAST_DELEGATE_TwoParams(FPythonLogCapture, EPythonLogOutputType, const TCHAR*);
+	FPythonLogCapture& GetPythonLogCapture();
 }
 
 #endif	// WITH_PYTHON
diff --git a/Engine/Plugins/Experimental/PythonScriptPlugin/Source/PythonScriptPlugin/Private/PyOnlineDocsWriter.cpp b/Engine/Plugins/Experimental/PythonScriptPlugin/Source/PythonScriptPlugin/Private/PyOnlineDocsWriter.cpp
--- a/Engine/Plugins/Experimental/PythonScriptPlugin/Source/PythonScriptPlugin/Private/PyOnlineDocsWriter.cpp
+++ b/Engine/Plugins/Experimental/PythonScriptPlugin/Source/PythonScriptPlugin/Private/PyOnlineDocsWriter.cpp
@@ -1,7 +1,8 @@
 // Copyright 1998-2019 Epic Games, Inc. All Rights Reserved.
 
 #include "PyOnlineDocsWriter.h"
-#include "PythonScriptPlugin.h"
+#include "IPythonScriptPlugin.h"
+#include "PyUtil.h"
 #include "PyGenUtil.h"
 #include "HAL/FileManager.h"
 #include "Logging/MessageLog.h"
@@ -380,7 +381,7 @@ void FPyOnlineDocsWriter::GenerateFiles(const FString& InPythonStubPath)
 #endif // !NO_LOGGING
 
 		// Run the Python commands
-		bool PyRunSuccess = FPythonScriptPlugin::Get()->RunString(*PyCommandStr);
+		bool PyRunSuccess = IPythonScriptPlugin::Get()->ExecPythonCommand(*PyCommandStr);
 
 #if !NO_LOGGING
 		if (!bLogSphinx)
diff --git a/Engine/Plugins/Experimental/PythonScriptPlugin/Source/PythonScriptPlugin/Private/PyUtil.cpp b/Engine/Plugins/Experimental/PythonScriptPlugin/Source/PythonScriptPlugin/Private/PyUtil.cpp
--- a/Engine/Plugins/Experimental/PythonScriptPlugin/Source/PythonScriptPlugin/Private/PyUtil.cpp
+++ b/Engine/Plugins/Experimental/PythonScriptPlugin/Source/PythonScriptPlugin/Private/PyUtil.cpp
@@ -90,6 +90,16 @@ FString PyStringToUEString(PyObject* InPyStr)
 	return Str;
 }
 
+FString PyObjectToUEStringRepr(PyObject* InPyObj)
+{
+	FPyObjectPtr PyReprObj = FPyObjectPtr::StealReference(PyObject_Repr(InPyObj));
+	if (PyReprObj)
+	{
+		return PyStringToUEString(PyReprObj);
+	}
+	return PyObjectToUEString(InPyObj);
+}
+
 FPropValueOnScope::FPropValueOnScope(const UProperty* InProp)
 	: Prop(InProp)
 	, Value(nullptr)
@@ -1257,7 +1267,7 @@ FString BuildPythonError()
 	return PythonErrorString;
 }
 
-void LogPythonError(const bool bInteractive)
+FString LogPythonError(const bool bInteractive)
 {
 	const FString ErrorStr = BuildPythonError();
 
@@ -1281,9 +1291,11 @@ void LogPythonError(const bool bInteractive)
 			FMessageDialog::Open(EAppMsgType::Ok, FText::AsCultureInvariant(ErrorStr), &DlgTitle);
 		}
 	}
+
+	return ErrorStr;
 }
 
-void ReThrowPythonError()
+FString ReThrowPythonError()
 {
 	const FString ErrorStr = BuildPythonError();
 
@@ -1291,6 +1303,8 @@ void ReThrowPythonError()
 	{
 		FFrame::KismetExecutionMessage(*ErrorStr, ELogVerbosity::Error);
 	}
+
+	return ErrorStr;
 }
 
 }
diff --git a/Engine/Plugins/Experimental/PythonScriptPlugin/Source/PythonScriptPlugin/Private/PyUtil.h b/Engine/Plugins/Experimental/PythonScriptPlugin/Source/PythonScriptPlugin/Private/PyUtil.h
--- a/Engine/Plugins/Experimental/PythonScriptPlugin/Source/PythonScriptPlugin/Private/PyUtil.h
+++ b/Engine/Plugins/Experimental/PythonScriptPlugin/Source/PythonScriptPlugin/Private/PyUtil.h
@@ -57,6 +57,9 @@ namespace PyUtil
 	/** Given a Python string/unicode object, extract the string value into an FString */
 	FString PyStringToUEString(PyObject* InPyStr);
 
+	/** Given a Python object, convert it to its string representation (repr) and extract the string value into an FString */
+	FString PyObjectToUEStringRepr(PyObject* InPyObj);
+
 	/** Given two values, perform a rich-comparison and return the result */
 	template <typename T, typename U>
 	PyObject* PyRichCmp(const T& InLHS, const U& InRHS, const int InOp)
@@ -368,11 +371,17 @@ namespace PyUtil
 	/** Enable developer warnings (eg, deprecation warnings) */
 	bool EnableDeveloperWarnings();
 
-	/** Log any pending Python error (will also clear the error) */
-	void LogPythonError(const bool bInteractive = false);
+	/**
+	 * Log any pending Python error (will also clear the error).
+	 * @return The error text that was logged.
+	 */
+	FString LogPythonError(const bool bInteractive = false);
 
-	/** Re-throw any pending Python error via FFrame::KismetExecutionMessage (will also clear the error) */
-	void ReThrowPythonError();
+	/**
+	 * Re-throw any pending Python error via FFrame::KismetExecutionMessage (will also clear the error).
+	 * @return The error text that was thrown.
+	 */
+	FString ReThrowPythonError();
 }
 
 #endif	// WITH_PYTHON
diff --git a/Engine/Plugins/Experimental/PythonScriptPlugin/Source/PythonScriptPlugin/Private/PythonScriptCommandlet.cpp b/Engine/Plugins/Experimental/PythonScriptPlugin/Source/PythonScriptPlugin/Private/PythonScriptCommandlet.cpp
--- a/Engine/Plugins/Experimental/PythonScriptPlugin/Source/PythonScriptPlugin/Private/PythonScriptCommandlet.cpp
+++ b/Engine/Plugins/Experimental/PythonScriptPlugin/Source/PythonScriptPlugin/Private/PythonScriptCommandlet.cpp
@@ -1,7 +1,8 @@
 // Copyright 1998-2019 Epic Games, Inc. All Rights Reserved.
 
 #include "PythonScriptCommandlet.h"
-#include "PythonScriptPlugin.h"
+#include "IPythonScriptPlugin.h"
+#include "Logging/LogMacros.h"
 
 DEFINE_LOG_CATEGORY_STATIC(LogPythonScriptCommandlet, Log, All);
 
@@ -32,8 +33,14 @@ int32 UPythonScriptCommandlet::Main(const FString& Params)
 	}
 
 #if WITH_PYTHON
-	UE_LOG(LogPythonScriptCommandlet, Display, TEXT("Running Python script: %s"), *PythonScript);
-	FPythonScriptPlugin::Get()->HandlePythonExecCommand(*PythonScript);
+	{
+		UE_LOG(LogPythonScriptCommandlet, Display, TEXT("Running Python script: %s"), *PythonScript);
+
+		FPythonCommandEx PythonCommand;
+		PythonCommand.Flags |= EPythonCommandFlags::Unattended;
+		PythonCommand.Command = PythonScript;
+		IPythonScriptPlugin::Get()->ExecPythonCommandEx(PythonCommand);
+	}
 #else	// WITH_PYTHON
 	UE_LOG(LogPythonScriptCommandlet, Error, TEXT("Python script cannot run as the plugin was built as a stub!"));
 	return -1;
diff --git a/Engine/Plugins/Experimental/PythonScriptPlugin/Source/PythonScriptPlugin/Private/PythonScriptPlugin.cpp b/Engine/Plugins/Experimental/PythonScriptPlugin/Source/PythonScriptPlugin/Private/PythonScriptPlugin.cpp
--- a/Engine/Plugins/Experimental/PythonScriptPlugin/Source/PythonScriptPlugin/Private/PythonScriptPlugin.cpp
+++ b/Engine/Plugins/Experimental/PythonScriptPlugin/Source/PythonScriptPlugin/Private/PythonScriptPlugin.cpp
@@ -74,7 +74,7 @@ struct FPythonScopedArgv
 	TArray<PyUtil::FPyApiChar*> PyCommandLineArgPtrs;
 };
 
-FPythonCommandExecutor::FPythonCommandExecutor(FPythonScriptPlugin* InPythonScriptPlugin)
+FPythonCommandExecutor::FPythonCommandExecutor(IPythonScriptPlugin* InPythonScriptPlugin)
 	: PythonScriptPlugin(InPythonScriptPlugin)
 {
 }
@@ -97,12 +97,12 @@ FText FPythonCommandExecutor::GetDisplayName() const
 
 FText FPythonCommandExecutor::GetDescription() const
 {
-	return LOCTEXT("PythonCommandExecutorDescription", "Execute Python Scripts");
+	return LOCTEXT("PythonCommandExecutorDescription", "Execute Python scripts (including files)");
 }
 
 FText FPythonCommandExecutor::GetHintText() const
 {
-	return LOCTEXT("PythonCommandExecutorHintText", "Enter Python Script");
+	return LOCTEXT("PythonCommandExecutorHintText", "Enter Python script or a filename to execute");
 }
 
 void FPythonCommandExecutor::GetAutoCompleteSuggestions(const TCHAR* Input, TArray<FString>& Out)
@@ -119,7 +119,8 @@ bool FPythonCommandExecutor::Exec(const TCHAR* Input)
 	IConsoleManager::Get().AddConsoleHistoryEntry(TEXT("Python"), Input);
 
 	UE_LOG(LogPython, Log, TEXT("%s"), Input);
-	PythonScriptPlugin->HandlePythonExecCommand(Input);
+
+	PythonScriptPlugin->ExecPythonCommand(Input);
 
 	return true;
 }
@@ -134,6 +135,70 @@ bool FPythonCommandExecutor::AllowMultiLine() const
 	return true;
 }
 
+FPythonREPLCommandExecutor::FPythonREPLCommandExecutor(IPythonScriptPlugin* InPythonScriptPlugin)
+	: PythonScriptPlugin(InPythonScriptPlugin)
+{
+}
+
+FName FPythonREPLCommandExecutor::StaticName()
+{
+	static const FName CmdExecName = TEXT("PythonREPL");
+	return CmdExecName;
+}
+
+FName FPythonREPLCommandExecutor::GetName() const
+{
+	return StaticName();
+}
+
+FText FPythonREPLCommandExecutor::GetDisplayName() const
+{
+	return LOCTEXT("PythonREPLCommandExecutorDisplayName", "Python (REPL)");
+}
+
+FText FPythonREPLCommandExecutor::GetDescription() const
+{
+	return LOCTEXT("PythonREPLCommandExecutorDescription", "Execute a single Python statement and show its result");
+}
+
+FText FPythonREPLCommandExecutor::GetHintText() const
+{
+	return LOCTEXT("PythonREPLCommandExecutorHintText", "Enter a single Python statement");
+}
+
+void FPythonREPLCommandExecutor::GetAutoCompleteSuggestions(const TCHAR* Input, TArray<FString>& Out)
+{
+}
+
+void FPythonREPLCommandExecutor::GetExecHistory(TArray<FString>& Out)
+{
+	IConsoleManager::Get().GetConsoleHistory(TEXT("PythonREPL"), Out);
+}
+
+bool FPythonREPLCommandExecutor::Exec(const TCHAR* Input)
+{
+	IConsoleManager::Get().AddConsoleHistoryEntry(TEXT("PythonREPL"), Input);
+
+	UE_LOG(LogPython, Log, TEXT("%s"), Input);
+
+	FPythonCommandEx PythonCommand;
+	PythonCommand.ExecutionMode = EPythonCommandExecutionMode::ExecuteStatement;
+	PythonCommand.Command = Input;
+	PythonScriptPlugin->ExecPythonCommandEx(PythonCommand);
+
+	return true;
+}
+
+bool FPythonREPLCommandExecutor::AllowHotKeyClose() const
+{
+	return false;
+}
+
+bool FPythonREPLCommandExecutor::AllowMultiLine() const
+{
+	return true;
+}
+
 #if WITH_EDITOR
 class FPythonCommandMenuImpl : public IPythonCommandMenu
 {
@@ -337,6 +402,7 @@ private:
 FPythonScriptPlugin::FPythonScriptPlugin()
 #if WITH_PYTHON
 	: CmdExec(this)
+	, CmdREPLExec(this)
 	, CmdMenu(nullptr)
 	, bInitialized(false)
 	, bHasTicked(false)
@@ -350,11 +416,42 @@ bool FPythonScriptPlugin::IsPythonAvailable() const
 }
 
 bool FPythonScriptPlugin::ExecPythonCommand(const TCHAR* InPythonCommand)
+{
+	FPythonCommandEx PythonCommand;
+	PythonCommand.Command = InPythonCommand;
+	return ExecPythonCommandEx(PythonCommand);
+}
+
+bool FPythonScriptPlugin::ExecPythonCommandEx(FPythonCommandEx& InOutPythonCommand)
 {
 #if WITH_PYTHON
-	return HandlePythonExecCommand(InPythonCommand);
+	if (InOutPythonCommand.ExecutionMode == EPythonCommandExecutionMode::ExecuteFile)
+	{
+		// We may have been passed literal code or a file
+		// To work out which, extract the first token and see if it's a .py file
+		// If it is, treat the remaining text as arguments to the file
+		// Otherwise, treat it as literal code
+		FString ExtractedFilename;
+		{
+			const TCHAR* Tmp = *InOutPythonCommand.Command;
+			ExtractedFilename = FParse::Token(Tmp, false);
+		}
+		if (FPaths::GetExtension(ExtractedFilename) == TEXT("py"))
+		{
+			return RunFile(*ExtractedFilename, *InOutPythonCommand.Command, InOutPythonCommand);
+		}
+		else
+		{
+			return RunString(InOutPythonCommand);
+		}
+	}
+	else
+	{
+		return RunString(InOutPythonCommand);
+	}
 #else	// WITH_PYTHON
-	ensureAlwaysMsgf(false, TEXT("Python is not available!"));
+	InOutPythonCommand.CommandResult = TEXT("Python is not available!");
+	ensureAlwaysMsgf(false, TEXT("%s"), *InOutPythonCommand.CommandResult);
 	return false;
 #endif	// WITH_PYTHON
 }
@@ -374,11 +471,12 @@ void FPythonScriptPlugin::StartupModule()
 #if WITH_PYTHON
 	InitializePython();
 	IModularFeatures::Get().RegisterModularFeature(IConsoleCommandExecutor::ModularFeatureName(), &CmdExec);
+	IModularFeatures::Get().RegisterModularFeature(IConsoleCommandExecutor::ModularFeatureName(), &CmdREPLExec);
 
 #if WITH_EDITOR
-		check(CmdMenu == nullptr);
-		CmdMenu = new FPythonCommandMenuImpl();
-		CmdMenu->OnStartupMenu();
+	check(CmdMenu == nullptr);
+	CmdMenu = new FPythonCommandMenuImpl();
+	CmdMenu->OnStartupMenu();
 #endif // WITH_EDITOR
 
 	FCoreDelegates::OnPreExit.AddRaw(this, &FPythonScriptPlugin::ShutdownPython);
@@ -398,6 +496,7 @@ void FPythonScriptPlugin::ShutdownModule()
 #endif // WITH_EDITOR
 
 	IModularFeatures::Get().UnregisterModularFeature(IConsoleCommandExecutor::ModularFeatureName(), &CmdExec);
+	IModularFeatures::Get().UnregisterModularFeature(IConsoleCommandExecutor::ModularFeatureName(), &CmdREPLExec);
 	ShutdownPython();
 #endif	// WITH_PYTHON
 }
@@ -407,7 +506,7 @@ bool FPythonScriptPlugin::Exec(UWorld* InWorld, const TCHAR* Cmd, FOutputDevice&
 #if WITH_PYTHON
 	if (FParse::Command(&Cmd, TEXT("PY")))
 	{
-		HandlePythonExecCommand(Cmd);
+		ExecPythonCommand(Cmd);
 		return true;
 	}
 #endif	// WITH_PYTHON
@@ -638,12 +737,13 @@ void FPythonScriptPlugin::Tick(const float InDeltaTime)
 			const FString PotentialFilePath = PySysPath / TEXT("init_unreal.py");
 			if (FPaths::FileExists(PotentialFilePath))
 			{
-				RunFile(*PotentialFilePath, nullptr);
+				// Quote the path in case it contains spaces so that the token parsing will work as expected
+				ExecPythonCommand(*FString::Printf(TEXT("\"%s\""), *PotentialFilePath));
 			}
 		}
 		for (const FString& StartupScript : GetDefault<UPythonScriptPluginSettings>()->StartupScripts)
 		{
-			HandlePythonExecCommand(*StartupScript);
+			ExecPythonCommand(*StartupScript);
 		}
 
 #if WITH_EDITOR
@@ -718,27 +818,6 @@ void FPythonScriptPlugin::ImportUnrealModule(const TCHAR* InModuleName)
 	}
 }
 
-bool FPythonScriptPlugin::HandlePythonExecCommand(const TCHAR* InPythonCommand)
-{
-	// We may have been passed literal code or a file
-	// To work out which, extract the first token and see if it's a .py file
-	// If it is, treat the remaining text as arguments to the file
-	// Otherwise, treat it as literal code
-	FString ExtractedFilename;
-	{
-		const TCHAR* Tmp = InPythonCommand;
-		ExtractedFilename = FParse::Token(Tmp, false);
-	}
-	if (FPaths::GetExtension(ExtractedFilename) == TEXT("py"))
-	{
-		return RunFile(*ExtractedFilename, InPythonCommand);
-	}
-	else
-	{
-		return RunString(InPythonCommand);
-	}
-}
-
 PyObject* FPythonScriptPlugin::EvalString(const TCHAR* InStr, const TCHAR* InContext, const int InMode)
 {
 	return EvalString(InStr, InContext, InMode, PyConsoleGlobalDict, PyConsoleLocalDict);
@@ -771,16 +850,41 @@ PyObject* FPythonScriptPlugin::EvalString(const TCHAR* InStr, const TCHAR* InCon
 	return PyEval_EvalCode((PyUtil::FPyCodeObjectType*)PyCodeObj.Get(), InGlobalDict, InLocalDict);
 }
 
-bool FPythonScriptPlugin::RunString(const TCHAR* InStr)
+bool FPythonScriptPlugin::RunString(FPythonCommandEx& InOutPythonCommand)
 {
 	// Execute Python code within this block
 	{
 		FPyScopedGIL GIL;
+		TGuardValue<bool> UnattendedScriptGuard(GIsRunningUnattendedScript, GIsRunningUnattendedScript || EnumHasAnyFlags(InOutPythonCommand.Flags, EPythonCommandFlags::Unattended));
 
-		FPyObjectPtr PyResult = FPyObjectPtr::StealReference(EvalString(InStr, TEXT("<string>"), Py_file_input));
-		if (!PyResult)
+		int PyExecMode = 0;
+		switch (InOutPythonCommand.ExecutionMode)
 		{
-			PyUtil::LogPythonError();
+		case EPythonCommandExecutionMode::ExecuteFile:
+			PyExecMode = Py_file_input;
+			break;
+		case EPythonCommandExecutionMode::ExecuteStatement:
+			PyExecMode = Py_single_input;
+			break;
+		case EPythonCommandExecutionMode::EvaluateStatement:
+			PyExecMode = Py_eval_input;
+			break;
+		default:
+			checkf(false, TEXT("Invalid EPythonCommandExecutionMode!"));
+			break;
+		}
+
+		FDelegateHandle LogCaptureHandle = PyCore::GetPythonLogCapture().AddLambda([&InOutPythonCommand](EPythonLogOutputType InLogType, const TCHAR* InLogString) { InOutPythonCommand.LogOutput.Add(FPythonLogOutputEntry{ InLogType, InLogString }); });
+		FPyObjectPtr PyResult = FPyObjectPtr::StealReference(EvalString(*InOutPythonCommand.Command, TEXT("<string>"), PyExecMode));
+		PyCore::GetPythonLogCapture().Remove(LogCaptureHandle);
+		
+		if (PyResult)
+		{
+			InOutPythonCommand.CommandResult = PyUtil::PyObjectToUEStringRepr(PyResult);
+		}
+		else
+		{
+			InOutPythonCommand.CommandResult = PyUtil::LogPythonError();
 			return false;
 		}
 	}
@@ -789,7 +893,7 @@ bool FPythonScriptPlugin::RunString(const TCHAR* InStr)
 	return true;
 }
 
-bool FPythonScriptPlugin::RunFile(const TCHAR* InFile, const TCHAR* InArgs)
+bool FPythonScriptPlugin::RunFile(const TCHAR* InFile, const TCHAR* InArgs, FPythonCommandEx& InOutPythonCommand)
 {
 	auto ResolveFilePath = [InFile]() -> FString
 	{
@@ -832,7 +936,8 @@ bool FPythonScriptPlugin::RunFile(const TCHAR* InFile, const TCHAR* InArgs)
 
 	if (!bLoaded)
 	{
-		UE_LOG(LogPython, Error, TEXT("Could not load Python file '%s' (resolved from '%s')"), *ResolvedFilePath, InFile);
+		InOutPythonCommand.CommandResult = FString::Printf(TEXT("Could not load Python file '%s' (resolved from '%s')"), *ResolvedFilePath, InFile);
+		UE_LOG(LogPython, Error, TEXT("%s"), *InOutPythonCommand.CommandResult);
 		return false;
 	}
 
@@ -840,6 +945,7 @@ bool FPythonScriptPlugin::RunFile(const TCHAR* InFile, const TCHAR* InArgs)
 	double ElapsedSeconds = 0.0;
 	{
 		FPyScopedGIL GIL;
+		TGuardValue<bool> UnattendedScriptGuard(GIsRunningUnattendedScript, GIsRunningUnattendedScript || EnumHasAnyFlags(InOutPythonCommand.Flags, EPythonCommandFlags::Unattended));
 
 		FPyObjectPtr PyFileGlobalDict = FPyObjectPtr::StealReference(PyDict_Copy(PyDefaultGlobalDict));
 		FPyObjectPtr PyFileLocalDict = PyFileGlobalDict;
@@ -856,13 +962,18 @@ bool FPythonScriptPlugin::RunFile(const TCHAR* InFile, const TCHAR* InArgs)
 			FScopedDurationTimer Timer(ElapsedSeconds);
 			FPythonScopedArgv ScopedArgv(InArgs);
 
-			// We can't just use PyRun_File here as Python isn't always built against the same version of the CRT as UE4, so we get a crash at the CRT layer
-			PyResult = FPyObjectPtr::StealReference(EvalString(*FileStr, *ResolvedFilePath, Py_file_input, PyFileGlobalDict, PyFileLocalDict));
+			FDelegateHandle LogCaptureHandle = PyCore::GetPythonLogCapture().AddLambda([&InOutPythonCommand](EPythonLogOutputType InLogType, const TCHAR* InLogString) { InOutPythonCommand.LogOutput.Add(FPythonLogOutputEntry{ InLogType, InLogString }); });
+			PyResult = FPyObjectPtr::StealReference(EvalString(*FileStr, *ResolvedFilePath, Py_file_input, PyFileGlobalDict, PyFileLocalDict)); // We can't just use PyRun_File here as Python isn't always built against the same version of the CRT as UE4, so we get a crash at the CRT layer
+			PyCore::GetPythonLogCapture().Remove(LogCaptureHandle);
 		}
 
-		if (!PyResult)
+		if (PyResult)
+		{
+			InOutPythonCommand.CommandResult = PyUtil::PyObjectToUEStringRepr(PyResult);
+		}
+		else
 		{
-			PyUtil::LogPythonError();
+			InOutPythonCommand.CommandResult = PyUtil::LogPythonError();
 			return false;
 		}
 	}
diff --git a/Engine/Plugins/Experimental/PythonScriptPlugin/Source/PythonScriptPlugin/Private/PythonScriptPlugin.h b/Engine/Plugins/Experimental/PythonScriptPlugin/Source/PythonScriptPlugin/Private/PythonScriptPlugin.h
--- a/Engine/Plugins/Experimental/PythonScriptPlugin/Source/PythonScriptPlugin/Private/PythonScriptPlugin.h
+++ b/Engine/Plugins/Experimental/PythonScriptPlugin/Source/PythonScriptPlugin/Private/PythonScriptPlugin.h
@@ -20,7 +20,7 @@ class FPythonScriptPlugin;
 class FPythonCommandExecutor : public IConsoleCommandExecutor
 {
 public:
-	FPythonCommandExecutor(FPythonScriptPlugin* InPythonScriptPlugin);
+	FPythonCommandExecutor(IPythonScriptPlugin* InPythonScriptPlugin);
 
 	static FName StaticName();
 	virtual FName GetName() const override;
@@ -37,7 +37,33 @@ public:
 		return FInputChord();
 	}
 private:
-	FPythonScriptPlugin* PythonScriptPlugin;
+	IPythonScriptPlugin* PythonScriptPlugin;
+};
+
+/**
+ * Executor for "Python (REPL)" commands
+ */
+class FPythonREPLCommandExecutor : public IConsoleCommandExecutor
+{
+public:
+	FPythonREPLCommandExecutor(IPythonScriptPlugin* InPythonScriptPlugin);
+
+	static FName StaticName();
+	virtual FName GetName() const override;
+	virtual FText GetDisplayName() const override;
+	virtual FText GetDescription() const override;
+	virtual FText GetHintText() const override;
+	virtual void GetAutoCompleteSuggestions(const TCHAR* Input, TArray<FString>& Out) override;
+	virtual void GetExecHistory(TArray<FString>& Out) override;
+	virtual bool Exec(const TCHAR* Input) override;
+	virtual bool AllowHotKeyClose() const override;
+	virtual bool AllowMultiLine() const override;
+	virtual FInputChord GetHotKey() const override
+	{
+		return FInputChord();
+	}
+private:
+	IPythonScriptPlugin* PythonScriptPlugin;
 };
 
 /**
@@ -68,6 +94,7 @@ public:
 	//~ IPythonScriptPlugin interface
 	virtual bool IsPythonAvailable() const override;
 	virtual bool ExecPythonCommand(const TCHAR* InPythonCommand) override;
+	virtual bool ExecPythonCommandEx(FPythonCommandEx& InOutPythonCommand) override;
 	virtual FSimpleMulticastDelegate& OnPythonInitialized() override;
 	virtual FSimpleMulticastDelegate& OnPythonShutdown() override;
 
@@ -85,14 +112,15 @@ public:
 	 */
 	void ImportUnrealModule(const TCHAR* InModuleName);
 
-	bool HandlePythonExecCommand(const TCHAR* InPythonCommand);
-
+	/** Evaluate/Execute a Python string, and return the result */
 	PyObject* EvalString(const TCHAR* InStr, const TCHAR* InContext, const int InMode);
 	PyObject* EvalString(const TCHAR* InStr, const TCHAR* InContext, const int InMode, PyObject* InGlobalDict, PyObject* InLocalDict);
 
-	bool RunString(const TCHAR* InStr);
+	/** Run literal Python script */
+	bool RunString(FPythonCommandEx& InOutPythonCommand);
 
-	bool RunFile(const TCHAR* InFile, const TCHAR* InArgs);
+	/** Run a Python file */
+	bool RunFile(const TCHAR* InFile, const TCHAR* InArgs, FPythonCommandEx& InOutPythonCommand);
 #endif	// WITH_PYTHON
 
 private:
@@ -120,6 +148,7 @@ private:
 #endif	// WITH_EDITOR
 
 	FPythonCommandExecutor CmdExec;
+	FPythonREPLCommandExecutor CmdREPLExec;
 	IPythonCommandMenu* CmdMenu;
 	FDelegateHandle TickHandle;
 	FDelegateHandle ModuleDelayedHandle;
diff --git a/Engine/Plugins/Experimental/PythonScriptPlugin/Source/PythonScriptPlugin/Public/IPythonScriptPlugin.h b/Engine/Plugins/Experimental/PythonScriptPlugin/Source/PythonScriptPlugin/Public/IPythonScriptPlugin.h
--- a/Engine/Plugins/Experimental/PythonScriptPlugin/Source/PythonScriptPlugin/Public/IPythonScriptPlugin.h
+++ b/Engine/Plugins/Experimental/PythonScriptPlugin/Source/PythonScriptPlugin/Public/IPythonScriptPlugin.h
@@ -3,8 +3,11 @@
 #pragma once
 
 #include "CoreTypes.h"
-#include "Modules/ModuleInterface.h"
+#include "Misc/EnumClassFlags.h"
 #include "Modules/ModuleManager.h"
+#include "Modules/ModuleInterface.h"
+
+struct FPythonCommandEx;
 
 class IPythonScriptPlugin : public IModuleInterface
 {
@@ -27,6 +30,12 @@ public:
 	 * @return true if the command ran successfully, false if there were errors (the output log will show the errors).
 	 */
 	virtual bool ExecPythonCommand(const TCHAR* InPythonCommand) = 0;
+
+	/**
+	 * Execute the given Python command.
+	 * @return true if the command ran successfully, false if there were errors.
+	 */
+	virtual bool ExecPythonCommandEx(FPythonCommandEx& InOutPythonCommand) = 0;
 	
 	/**
 	 * Delegate called after Python has been initialized.
@@ -38,3 +47,122 @@ public:
 	 */
 	virtual FSimpleMulticastDelegate& OnPythonShutdown() = 0;
 };
+
+/** Types of log output that Python can give. */
+enum class EPythonLogOutputType : uint8
+{
+	/** This log was informative. */
+	Info,
+	/** This log was a warning. */
+	Warning,
+	/** This log was an error. */
+	Error,
+};
+
+/** Flags that can be specified when running Python commands. */
+enum class EPythonCommandFlags : uint8
+{
+	/** No special behavior. */
+	None = 0,
+	/** Run the Python command in "unattended" mode (GIsRunningUnattendedScript set to true), which will suppress certain pieces of UI. */
+	Unattended = 1<<0,
+};
+ENUM_CLASS_FLAGS(EPythonCommandFlags);
+
+/** Controls the execution mode used for the Python command. */
+enum class EPythonCommandExecutionMode : uint8
+{
+	/** Execute the Python command as a file. This allows you to execute either a literal Python script containing multiple statements, or a file with optional arguments. */
+	ExecuteFile,
+	/** Execute the Python command as a single statement. This will execute a single statement and print the result. This mode cannot run files. */
+	ExecuteStatement,
+	/** Evaluate the Python command as a single statement. This will evaluate a single statement and return the result. This mode cannot run files. */
+	EvaluateStatement,
+};
+
+/** Log output entry captured from Python. */
+struct FPythonLogOutputEntry
+{
+	/** The type of the log output. */
+	EPythonLogOutputType Type = EPythonLogOutputType::Info;
+
+	/** The log output string. */
+	FString Output;
+};
+
+/** Extended information when executing Python commands. */
+struct FPythonCommandEx
+{
+	/** Flags controlling how the command should be run. */
+	EPythonCommandFlags Flags = EPythonCommandFlags::None;
+
+	/** Controls the mode used to execute the command. */
+	EPythonCommandExecutionMode ExecutionMode = EPythonCommandExecutionMode::ExecuteFile;
+
+	/** The command to run. This may be literal Python code, or a file (with optional arguments) to run. */
+	FString Command;
+
+	/** The result of running the command. On success, for EvaluateStatement mode this will be the actual result of running the command, and will be None in all other cases. On failure, this will be the error information (typically a Python exception trace). */
+	FString CommandResult;
+
+	/** The log output captured while running the command. */
+	TArray<FPythonLogOutputEntry> LogOutput;
+};
+
+inline const TCHAR* LexToString(EPythonLogOutputType InType)
+{
+	switch (InType)
+	{
+	case EPythonLogOutputType::Info:
+		return TEXT("Info");
+	case EPythonLogOutputType::Warning:
+		return TEXT("Warning");
+	case EPythonLogOutputType::Error:
+		return TEXT("Error");
+	default:
+		break;
+	}
+	return TEXT("<Unknown EPythonLogOutputType>");
+}
+
+inline const TCHAR* LexToString(EPythonCommandExecutionMode InMode)
+{
+	switch (InMode)
+	{
+	case EPythonCommandExecutionMode::ExecuteFile:
+		return TEXT("ExecuteFile");
+	case EPythonCommandExecutionMode::ExecuteStatement:
+		return TEXT("ExecuteStatement");
+	case EPythonCommandExecutionMode::EvaluateStatement:
+		return TEXT("EvaluateStatement");
+	default:
+		break;
+	}
+	return TEXT("<Unknown EPythonCommandExecutionMode>");
+}
+
+inline bool LexTryParseString(EPythonCommandExecutionMode& OutMode, const TCHAR* InBuffer)
+{
+	if (FCString::Stricmp(InBuffer, TEXT("ExecuteFile")) == 0)
+	{
+		OutMode = EPythonCommandExecutionMode::ExecuteFile;
+		return true;
+	}
+	if (FCString::Stricmp(InBuffer, TEXT("ExecuteStatement")) == 0)
+	{
+		OutMode = EPythonCommandExecutionMode::ExecuteStatement;
+		return true;
+	}
+	if (FCString::Stricmp(InBuffer, TEXT("EvaluateStatement")) == 0)
+	{
+		OutMode = EPythonCommandExecutionMode::EvaluateStatement;
+		return true;
+	}
+	return false;
+}
+
+inline void LexFromString(EPythonCommandExecutionMode& OutMode, const TCHAR* InBuffer)
+{
+	OutMode = EPythonCommandExecutionMode::ExecuteFile;
+	LexTryParseString(OutMode, InBuffer);
+}
