Added support for connecting to an existing editor instance to run Python commands on it remotely

This is off by default, and discovery is limited to the local host by default for safety. To enable it you need to turn on remote discovery in the Python plug-in project settings, which will make your editor discoverable for remote execution (via external scripts using the remote_execution.py module).

Remote execution uses a protocol of UTF-8 encoded JSON, and a combination of UDP-based multicast for discovery and remote connection requests, and TCP-based peer-to-peer communication for the actual command channel. See PythonScriptRemoteExecution.cpp for the full protocol definition.

remote_execution.py serves as both a reference implementation for the protocol implementation needed to connect to an editor instance, and as an out-of-the-box solution for applications that are written in (or can execute) Python (by importing it as a module and creating a RemoteExecution instance). remote_execution.py also provides a minimal example at the bottom of the file that demonstrates how to use the API, and this demo application can be seen by executing the file directly.

diff --git a/Engine/Plugins/Experimental/PythonScriptPlugin/Source/PythonScriptPlugin/Private/PythonScriptPlugin.cpp b/Engine/Plugins/Experimental/PythonScriptPlugin/Source/PythonScriptPlugin/Private/PythonScriptPlugin.cpp
--- a/Engine/Plugins/Experimental/PythonScriptPlugin/Source/PythonScriptPlugin/Private/PythonScriptPlugin.cpp
+++ b/Engine/Plugins/Experimental/PythonScriptPlugin/Source/PythonScriptPlugin/Private/PythonScriptPlugin.cpp
@@ -2,6 +2,7 @@
 
 #include "PythonScriptPlugin.h"
 #include "PythonScriptPluginSettings.h"
+#include "PythonScriptRemoteExecution.h"
 #include "PyGIL.h"
 #include "PyCore.h"
 #include "PySlate.h"
@@ -473,6 +474,9 @@ void FPythonScriptPlugin::StartupModule()
 	IModularFeatures::Get().RegisterModularFeature(IConsoleCommandExecutor::ModularFeatureName(), &CmdExec);
 	IModularFeatures::Get().RegisterModularFeature(IConsoleCommandExecutor::ModularFeatureName(), &CmdREPLExec);
 
+	check(!RemoteExecution);
+	RemoteExecution = MakeUnique<FPythonScriptRemoteExecution>(this);
+
 #if WITH_EDITOR
 	check(CmdMenu == nullptr);
 	CmdMenu = new FPythonCommandMenuImpl();
@@ -488,6 +492,8 @@ void FPythonScriptPlugin::ShutdownModule()
 #if WITH_PYTHON
 	FCoreDelegates::OnPreExit.RemoveAll(this);
 
+	RemoteExecution.Reset();
+
 #if WITH_EDITOR
 	check(CmdMenu);
 	CmdMenu->OnShutdownMenu();
@@ -752,9 +758,16 @@ void FPythonScriptPlugin::Tick(const float InDeltaTime)
 #endif	// WITH_EDITOR
 	}
 
+	RemoteExecution->Tick(InDeltaTime);
+
 	FPyWrapperTypeReinstancer::Get().ProcessPending();
 }
 
+void FPythonScriptPlugin::SyncRemoteExecutionToSettings()
+{
+	RemoteExecution->SyncToSettings();
+}
+
 void FPythonScriptPlugin::ImportUnrealModule(const TCHAR* InModuleName)
 {
 	const FString PythonModuleName = FString::Printf(TEXT("unreal_%s"), InModuleName);
diff --git a/Engine/Plugins/Experimental/PythonScriptPlugin/Source/PythonScriptPlugin/Private/PythonScriptPlugin.h b/Engine/Plugins/Experimental/PythonScriptPlugin/Source/PythonScriptPlugin/Private/PythonScriptPlugin.h
--- a/Engine/Plugins/Experimental/PythonScriptPlugin/Source/PythonScriptPlugin/Private/PythonScriptPlugin.h
+++ b/Engine/Plugins/Experimental/PythonScriptPlugin/Source/PythonScriptPlugin/Private/PythonScriptPlugin.h
@@ -11,6 +11,7 @@
 #include "Framework/Commands/InputChord.h"
 
 class FPythonScriptPlugin;
+class FPythonScriptRemoteExecution;
 
 #if WITH_PYTHON
 
@@ -106,6 +107,9 @@ public:
 	virtual bool Exec(UWorld* InWorld, const TCHAR* Cmd, FOutputDevice& Ar) override;
 
 #if WITH_PYTHON
+	/** Sync the remote execution environment to the current settings, starting or stopping it as required */
+	void SyncRemoteExecutionToSettings();
+
 	/** 
 	 * Import the given module into the "unreal" package.
 	 * This function will take the given name and attempt to import either "unreal_{name}" or "_unreal_{name}" into the "unreal" package as "unreal.{name}".
@@ -147,6 +151,7 @@ private:
 	void OnPrepareToCleanseEditorObject(UObject* InObject);
 #endif	// WITH_EDITOR
 
+	TUniquePtr<FPythonScriptRemoteExecution> RemoteExecution;
 	FPythonCommandExecutor CmdExec;
 	FPythonREPLCommandExecutor CmdREPLExec;
 	IPythonCommandMenu* CmdMenu;
diff --git a/Engine/Plugins/Experimental/PythonScriptPlugin/Source/PythonScriptPlugin/Private/PythonScriptPluginSettings.cpp b/Engine/Plugins/Experimental/PythonScriptPlugin/Source/PythonScriptPlugin/Private/PythonScriptPluginSettings.cpp
--- a/Engine/Plugins/Experimental/PythonScriptPlugin/Source/PythonScriptPlugin/Private/PythonScriptPluginSettings.cpp
+++ b/Engine/Plugins/Experimental/PythonScriptPlugin/Source/PythonScriptPlugin/Private/PythonScriptPluginSettings.cpp
@@ -1,6 +1,7 @@
 // Copyright 1998-2019 Epic Games, Inc. All Rights Reserved.
 
 #include "PythonScriptPluginSettings.h"
+#include "PythonScriptPlugin.h"
 
 #define LOCTEXT_NAMESPACE "PythonScriptPlugin"
 
@@ -8,6 +9,53 @@ UPythonScriptPluginSettings::UPythonScriptPluginSettings()
 {
 	CategoryName = TEXT("Plugins");
 	SectionName  = TEXT("Python");
+
+	RemoteExecutionMulticastGroupEndpoint = TEXT("239.0.0.1:6766");
+	RemoteExecutionMulticastBindAddress = TEXT("0.0.0.0");
+	RemoteExecutionSendBufferSizeBytes = 2 * 1024 * 1024;
+	RemoteExecutionReceiveBufferSizeBytes = 2 * 1024 * 1024;
+	RemoteExecutionMulticastTtl = 0;
+}
+
+#if WITH_EDITOR
+
+bool UPythonScriptPluginSettings::CanEditChange(const UProperty* InProperty) const
+{
+	bool bCanEditChange = Super::CanEditChange(InProperty);
+
+	if (bCanEditChange && InProperty)
+	{
+		if (InProperty->GetFName() == GET_MEMBER_NAME_CHECKED(UPythonScriptPluginSettings, RemoteExecutionMulticastGroupEndpoint) ||
+			InProperty->GetFName() == GET_MEMBER_NAME_CHECKED(UPythonScriptPluginSettings, RemoteExecutionMulticastBindAddress) ||
+			InProperty->GetFName() == GET_MEMBER_NAME_CHECKED(UPythonScriptPluginSettings, RemoteExecutionSendBufferSizeBytes) || 
+			InProperty->GetFName() == GET_MEMBER_NAME_CHECKED(UPythonScriptPluginSettings, RemoteExecutionReceiveBufferSizeBytes) ||
+			InProperty->GetFName() == GET_MEMBER_NAME_CHECKED(UPythonScriptPluginSettings, RemoteExecutionMulticastTtl)
+			)
+		{
+			bCanEditChange &= bRemoteExecution;
+		}
+	}
+
+	return bCanEditChange;
+}
+
+void UPythonScriptPluginSettings::PostEditChangeProperty(struct FPropertyChangedEvent& PropertyChangedEvent)
+{
+	Super::PostEditChangeProperty(PropertyChangedEvent);
+
+	if (PropertyChangedEvent.MemberProperty)
+	{
+		if (PropertyChangedEvent.MemberProperty->GetFName() == GET_MEMBER_NAME_CHECKED(UPythonScriptPluginSettings, bRemoteExecution) ||
+			PropertyChangedEvent.MemberProperty->GetFName() == GET_MEMBER_NAME_CHECKED(UPythonScriptPluginSettings, RemoteExecutionMulticastGroupEndpoint) ||
+			PropertyChangedEvent.MemberProperty->GetFName() == GET_MEMBER_NAME_CHECKED(UPythonScriptPluginSettings, RemoteExecutionMulticastBindAddress) ||
+			PropertyChangedEvent.MemberProperty->GetFName() == GET_MEMBER_NAME_CHECKED(UPythonScriptPluginSettings, RemoteExecutionSendBufferSizeBytes) ||
+			PropertyChangedEvent.MemberProperty->GetFName() == GET_MEMBER_NAME_CHECKED(UPythonScriptPluginSettings, RemoteExecutionReceiveBufferSizeBytes) ||
+			PropertyChangedEvent.MemberProperty->GetFName() == GET_MEMBER_NAME_CHECKED(UPythonScriptPluginSettings, RemoteExecutionMulticastTtl)
+			)
+		{
+			FPythonScriptPlugin::Get()->SyncRemoteExecutionToSettings();
+		}
+	}
 }
 
 FText UPythonScriptPluginSettings::GetSectionText() const
@@ -15,4 +63,7 @@ FText UPythonScriptPluginSettings::GetSectionText() const
 	return LOCTEXT("SettingsDisplayName", "Python");
 }
 
+#endif	// WITH_EDITOR
+
 #undef LOCTEXT_NAMESPACE
+
diff --git a/Engine/Plugins/Experimental/PythonScriptPlugin/Source/PythonScriptPlugin/Private/PythonScriptPluginSettings.h b/Engine/Plugins/Experimental/PythonScriptPlugin/Source/PythonScriptPlugin/Private/PythonScriptPluginSettings.h
--- a/Engine/Plugins/Experimental/PythonScriptPlugin/Source/PythonScriptPlugin/Private/PythonScriptPluginSettings.h
+++ b/Engine/Plugins/Experimental/PythonScriptPlugin/Source/PythonScriptPlugin/Private/PythonScriptPluginSettings.h
@@ -19,6 +19,10 @@ public:
 	UPythonScriptPluginSettings();
 
 #if WITH_EDITOR
+	//~ UObject interface
+	virtual bool CanEditChange(const UProperty* InProperty) const override;
+	virtual void PostEditChangeProperty(struct FPropertyChangedEvent& PropertyChangedEvent) override;
+
 	//~ UDeveloperSettings interface
 	virtual FText GetSectionText() const override;
 #endif
@@ -34,4 +38,28 @@ public:
 	/** Should Developer Mode be enabled on the Python interpreter (will enable extra warnings (eg, for deprecated code), and enable stub code generation for use with external IDEs). */
 	UPROPERTY(config, EditAnywhere, Category=Python, meta=(ConfigRestartRequired=true))
 	bool bDeveloperMode;
+
+	/** Should remote Python execution be enabled? */
+	UPROPERTY(config, EditAnywhere, Category=PythonRemoteExecution, meta=(DisplayName="Enable Remote Execution?"))
+	bool bRemoteExecution;
+
+	/** The multicast group endpoint (in the form of IP_ADDRESS:PORT_NUMBER) that the UDP multicast socket should join */
+	UPROPERTY(config, EditAnywhere, Category=PythonRemoteExecution, AdvancedDisplay, meta=(DisplayName="Multicast Group Endpoint"))
+	FString RemoteExecutionMulticastGroupEndpoint;
+
+	/** The adapter address that the UDP multicast socket should bind to, or 0.0.0.0 to bind to all adapters */
+	UPROPERTY(config, EditAnywhere, Category=PythonRemoteExecution, AdvancedDisplay, meta=(DisplayName="Multicast Bind Address"))
+	FString RemoteExecutionMulticastBindAddress;
+
+	/** Size of the send buffer for the remote endpoint connection */
+	UPROPERTY(config, EditAnywhere, Category=PythonRemoteExecution, AdvancedDisplay, meta=(DisplayName="Send Buffer Size", Units="Bytes"))
+	int32 RemoteExecutionSendBufferSizeBytes;
+
+	/** Size of the receive buffer for the remote endpoint connection */
+	UPROPERTY(config, EditAnywhere, Category=PythonRemoteExecution, AdvancedDisplay, meta=(DisplayName="Receive Buffer Size", Units="Bytes"))
+	int32 RemoteExecutionReceiveBufferSizeBytes;
+
+	/** The TTL that the UDP multicast socket should use (0 is limited to the local host, 1 is limited to the local subnet) */
+	UPROPERTY(config, EditAnywhere, Category=PythonRemoteExecution, AdvancedDisplay, meta=(DisplayName="Multicast Time-To-Live"))
+	uint8 RemoteExecutionMulticastTtl;
 };
diff --git a/Engine/Plugins/Experimental/PythonScriptPlugin/Source/PythonScriptPlugin/Private/PythonScriptRemoteExecution.cpp b/Engine/Plugins/Experimental/PythonScriptPlugin/Source/PythonScriptPlugin/Private/PythonScriptRemoteExecution.cpp
new file mode 100644
--- /dev/null
+++ b/Engine/Plugins/Experimental/PythonScriptPlugin/Source/PythonScriptPlugin/Private/PythonScriptRemoteExecution.cpp
@@ -0,0 +1,799 @@
+// Copyright 1998-2019 Epic Games, Inc. All Rights Reserved.
+
+#include "PythonScriptRemoteExecution.h"
+#include "IPythonScriptPlugin.h"
+#include "PythonScriptPluginSettings.h"
+
+#include "Misc/App.h"
+#include "Misc/Paths.h"
+#include "Misc/EngineVersion.h"
+#include "Logging/LogMacros.h"
+
+#include "Sockets.h"
+#include "Common/TcpSocketBuilder.h"
+#include "Common/UdpSocketBuilder.h"
+
+#include "Dom/JsonValue.h"
+#include "Dom/JsonObject.h"
+#include "Serialization/JsonSerializer.h"
+#include "Policies/CondensedJsonPrintPolicy.h"
+
+#define LOCTEXT_NAMESPACE "PythonScriptRemoteExecution"
+
+#if WITH_PYTHON
+
+/**
+ * Remote Execution protocol definition.
+ *
+ * Messages are a UTF-8 encoded JSON object with the following layout:
+ *	{
+ *		"version": 1,		// Required - Integer - Protocol version number (currently 1)
+ *		"magic": "ue_py",	// Required - String  - Protocol magic identifier (currently "ue_py")
+ * 		"type": "...",		// Required - String  - Identifier for the message type (see below)
+ * 		"source": "...",	// Required - String  - Unique identifier of the node that sent the message (typically a GUID)
+ * 		"dest": "...",		// Optional - String  - Unique identifier of the node that should receive the message (typically a GUID)
+ * 		"data": "..."		// Optional - Any	  - Message specific payload data (see below)
+ *	}
+ *
+ * Message types:
+ *	"ping" (UDP) - Service discovery request
+ *
+ *	"pong" (UDP) - Service discovery response
+ *		"dest" - Required - String - The node that requested the discovery
+ *		"data" - Required - Object - Payload data:
+ *			{
+ *				"user": "...",				// Required - String - The name of the user that owns the engine instance
+ *				"machine": "...",			// Required - String - The name of the machine that owns the engine instance
+ *				"engine_version": "...",	// Required - String - The version number of the engine
+ *				"engine_root": "...",		// Required - String - The absolute on-disk root of the engine
+ *				"project_root: "...",		// Optional - String - The absolute on-disk root of the project (if a project is loaded)
+ *				"project_name": "..."		// Optional - String - The name of the project (if a project is loaded)
+ *			}
+ *
+ *	"open_connection" (UDP) - Open a TCP command connection with the requested server
+ *		"dest" - Required - String - The node that should open a connection to the server
+ *		"data" - Required - Object - Payload data:
+ *			{
+ *				"command_ip": "...",		// Required - String  - The IP address of the TCP command server
+ *				"command_port": 0			// Required - Integer - The port of the TCP command server
+ *			}
+ *
+ *	"close_connection" (UDP) - Close any active TCP command connection
+ *		"dest" - Required - String - The node that should close its connection to the server
+ *
+ *	"command" (TCP) - Execute a remote Python command
+ *		"dest" - Required - String - The node that should execute the command
+ *		"data" - Required - Object - Payload data:
+ *			{
+ *				"command": "...",			// Required - String  - The command to execute
+ *				"unattended": True,			// Required - Boolean - Whether to run the command in "unattended" mode (suppressing some UI)
+ *				"exec_mode": "..."			// Required - String  - The execution mode to use as a string value (see EPythonCommandExecutionMode, and LexToString, LexFromString)
+ *			}
+ *
+ *	"command_result" (TCP) - Result of executing a remote Python command
+ *		"dest" - Required - String - The node that should receive the result
+ *		"data" - Required - Object - Payload data:
+ *			{
+ *				"success": True,			// Required - Boolean - Whether the command executed successfully (without raising an uncaught exception)
+ *				"command": "...",			// Required - String  - The command that was executed
+ *				"result": "...",			// Required - String  - The result of running the command
+ *				"output": [{				// Required - List    - The log output from running the command
+ *					"type: "...",				// Required - String - The type of output (Info, Warning, or Error)
+ *					"output": "..."				// Required - String - The actual log output
+ *					}]
+ *			}
+ */
+
+DEFINE_LOG_CATEGORY_STATIC(LogPythonRemoteExecution, Log, All);
+
+struct FPythonScriptRemoteExecutionMessage;
+
+namespace PythonScriptRemoteExecutionUtil
+{
+	const int32 ProtocolVersion = 1;
+	const FString ProtocolMagic = TEXT("ue_py");
+	const FString TypePing = TEXT("ping");
+	const FString TypePong = TEXT("pong");
+	const FString TypeOpenConnection = TEXT("open_connection");
+	const FString TypeCloseConnection = TEXT("close_connection");
+	const FString TypeCommand = TEXT("command");
+	const FString TypeCommandResult = TEXT("command_result");
+
+	bool SerializeJsonObject(const TSharedRef<FJsonObject>& InJsonObject, FString& OutJsonString, FText* OutError = nullptr);
+	bool DeserializeJsonObject(const FString& InJsonString, TSharedPtr<FJsonObject>& OutJsonObject, FText* OutError = nullptr);
+	bool ParseMessageData(const char* InMessageData, FPythonScriptRemoteExecutionMessage& OutMessage);
+	bool ParseMessageData(const TCHAR* InMessageData, FPythonScriptRemoteExecutionMessage& OutMessage);
+	bool WriteMessageData(const FPythonScriptRemoteExecutionMessage& InMessage, FString& OutMessageData);
+	void DestroySocket(ISocketSubsystem* InSocketSubsystem, FSocket*& InOutSocket);
+}
+
+struct FPythonScriptRemoteExecutionMessage
+{
+	FPythonScriptRemoteExecutionMessage() = default;
+
+	FPythonScriptRemoteExecutionMessage(FString InType, FString InSource, FString InDest = FString(), TSharedPtr<FJsonValue> InData = nullptr)
+		: Type(MoveTemp(InType))
+		, Source(MoveTemp(InSource))
+		, Dest(MoveTemp(InDest))
+		, Data(MoveTemp(InData))
+	{
+	}
+
+	bool PassesReceiveFilter(const FString& InNodeId) const
+	{
+		return !Source.Equals(InNodeId) && (Dest.IsEmpty() || Dest.Equals(InNodeId));
+	}
+
+	bool ToJsonObject(TSharedPtr<FJsonObject>& OutJsonObject, FText* OutError = nullptr) const
+	{
+		// Validate required fields
+		if (Type.IsEmpty())
+		{
+			if (OutError)
+			{
+				*OutError = LOCTEXT("Message.ToJson.Error.TypeCannotBeEmpty", "'type' cannot be empty!");
+			}
+			return false;
+		}
+		if (Source.IsEmpty())
+		{
+			if (OutError)
+			{
+				*OutError = LOCTEXT("Message.ToJson.Error.SourceCannotBeEmpty", "'source' cannot be empty!");
+			}
+			return false;
+		}
+
+		OutJsonObject = MakeShared<FJsonObject>();
+
+		// Write required protocol version information
+		OutJsonObject->SetNumberField(TEXT("version"), PythonScriptRemoteExecutionUtil::ProtocolVersion);
+		OutJsonObject->SetStringField(TEXT("magic"), PythonScriptRemoteExecutionUtil::ProtocolMagic);
+
+		// Write required fields
+		OutJsonObject->SetStringField(TEXT("type"), Type);
+		OutJsonObject->SetStringField(TEXT("source"), Source);
+
+		// Write optional fields
+		if (!Dest.IsEmpty())
+		{
+			OutJsonObject->SetStringField(TEXT("dest"), Dest);
+		}
+		if (Data.IsValid())
+		{
+			OutJsonObject->SetField(TEXT("data"), Data);
+		}
+
+		return true;
+	}
+
+	bool ToJsonString(FString& OutJsonString, FText* OutError = nullptr) const
+	{
+		TSharedPtr<FJsonObject> JsonObject;
+		return ToJsonObject(JsonObject, OutError) && PythonScriptRemoteExecutionUtil::SerializeJsonObject(JsonObject.ToSharedRef(), OutJsonString, OutError);
+	}
+
+	bool FromJsonObject(const TSharedRef<FJsonObject>& InJsonObject, FText* OutError = nullptr)
+	{
+		// Read and validate required protocol version information
+		{
+			int32 LocalVersion = 0;
+			if (!InJsonObject->TryGetNumberField(TEXT("version"), LocalVersion) || LocalVersion != PythonScriptRemoteExecutionUtil::ProtocolVersion)
+			{
+				if (OutError)
+				{
+					*OutError = FText::Format(LOCTEXT("Message.FromJson.Error.MissingOrInvalidVersion", "'version' is missing or incorrect (got {0}, expected {1})!"), LocalVersion, PythonScriptRemoteExecutionUtil::ProtocolVersion);
+				}
+				return false;
+			}
+
+			FString LocalMagic;
+			if (!InJsonObject->TryGetStringField(TEXT("magic"), LocalMagic) || !LocalMagic.Equals(PythonScriptRemoteExecutionUtil::ProtocolMagic))
+			{
+				if (OutError)
+				{
+					*OutError = FText::Format(LOCTEXT("Message.FromJson.Error.MissingOrInvalidMagic", "'magic' is missing or incorrect (got '{0}', expected '{1}')!"), FText::AsCultureInvariant(LocalMagic), FText::AsCultureInvariant(PythonScriptRemoteExecutionUtil::ProtocolMagic));
+				}
+				return false;
+			}
+		}
+
+		// Read required fields
+		{
+			FString LocalType;
+			if (!InJsonObject->TryGetStringField(TEXT("type"), LocalType))
+			{
+				if (OutError)
+				{
+					*OutError = LOCTEXT("Message.FromJson.Error.MissingType", "'type' is missing!");
+				}
+				return false;
+			}
+
+			FString LocalSource;
+			if (!InJsonObject->TryGetStringField(TEXT("source"), LocalSource))
+			{
+				if (OutError)
+				{
+					*OutError = LOCTEXT("Message.FromJson.Error.MissingSource", "'source' is missing!");
+				}
+				return false;
+			}
+
+			Type = MoveTemp(LocalType);
+			Source = MoveTemp(LocalSource);
+		}
+
+		// Read optional fields
+		InJsonObject->TryGetStringField(TEXT("dest"), Dest);
+		Data = InJsonObject->TryGetField(TEXT("data"));
+
+		return true;
+	}
+
+	bool FromJsonString(const FString& InJsonString, FText* OutError = nullptr)
+	{
+		TSharedPtr<FJsonObject> JsonObject;
+		return PythonScriptRemoteExecutionUtil::DeserializeJsonObject(InJsonString, JsonObject, OutError) && FromJsonObject(JsonObject.ToSharedRef(), OutError);
+	}
+
+	/** Identifier for the message type */
+	FString Type;
+
+	/** Unique identifier of the node that sent the message (typically a GUID) */
+	FString Source;
+
+	/** Unique identifier of the node that should receive the message (typically a GUID) */
+	FString Dest;
+
+	/** Message specific payload data */
+	TSharedPtr<FJsonValue> Data;
+};
+
+class FPythonScriptRemoteExecutionBroadcastConnection
+{
+public:
+	FPythonScriptRemoteExecutionBroadcastConnection(FPythonScriptRemoteExecution* InRemoteExecution)
+		: RemoteExecution(InRemoteExecution)
+		, BroadcastSocket(nullptr)
+	{
+		check(RemoteExecution);
+
+		NodeInfo = MakeShared<FJsonObject>();
+		NodeInfo->SetStringField(TEXT("user"), FApp::GetSessionOwner());
+		NodeInfo->SetStringField(TEXT("machine"), FPlatformProcess::ComputerName());
+		NodeInfo->SetStringField(TEXT("engine_version"), FEngineVersion::Current().ToString());
+		NodeInfo->SetStringField(TEXT("engine_root"), FPaths::ConvertRelativePathToFull(FPaths::EngineDir()));
+		if (FApp::HasProjectName())
+		{
+			NodeInfo->SetStringField(TEXT("project_root"), FPaths::ConvertRelativePathToFull(FPaths::ProjectDir()));
+			NodeInfo->SetStringField(TEXT("project_name"), FApp::GetProjectName());
+		}
+
+		const UPythonScriptPluginSettings* PythonScriptPluginSettings = GetDefault<UPythonScriptPluginSettings>();
+
+		bool bValidEndpointAddresses = true;
+
+		FIPv4Endpoint MulticastGroupEndpoint(FIPv4Address::Any, 0);
+		if (!FIPv4Endpoint::Parse(PythonScriptPluginSettings->RemoteExecutionMulticastGroupEndpoint, MulticastGroupEndpoint))
+		{
+			bValidEndpointAddresses = false;
+			UE_LOG(LogPythonRemoteExecution, Error, TEXT("Failed to parse multicast group endpoint '%s'. Remote execution will be disabled!"), *PythonScriptPluginSettings->RemoteExecutionMulticastGroupEndpoint);
+		}
+
+		FIPv4Address MulticastBindAddress = FIPv4Address::Any;
+		if (!FIPv4Address::Parse(PythonScriptPluginSettings->RemoteExecutionMulticastBindAddress, MulticastBindAddress))
+		{
+			bValidEndpointAddresses = false;
+			UE_LOG(LogPythonRemoteExecution, Error, TEXT("Failed to parse multicast bind address '%s'. Remote execution will be disabled!"), *PythonScriptPluginSettings->RemoteExecutionMulticastBindAddress);
+		}
+
+		if (bValidEndpointAddresses)
+		{
+			BroadcastGroupAddr = MulticastGroupEndpoint.ToInternetAddr();
+
+			BroadcastSocket = FUdpSocketBuilder(TEXT("PythonRemoteExecutionBroadcastSocket"))
+				.AsReusable()
+				.BoundToAddress(MulticastBindAddress)
+				.BoundToPort(MulticastGroupEndpoint.Port)
+#if PLATFORM_SUPPORTS_UDP_MULTICAST_GROUP
+				.JoinedToGroup(MulticastGroupEndpoint.Address)
+				.WithMulticastLoopback()
+#endif
+				.WithSendBufferSize(PythonScriptPluginSettings->RemoteExecutionSendBufferSizeBytes)
+				.WithReceiveBufferSize(PythonScriptPluginSettings->RemoteExecutionReceiveBufferSizeBytes)
+				.WithMulticastTtl(PythonScriptPluginSettings->RemoteExecutionMulticastTtl);
+		}
+	}
+
+	~FPythonScriptRemoteExecutionBroadcastConnection()
+	{
+		PythonScriptRemoteExecutionUtil::DestroySocket(RemoteExecution->SocketSubsystem, BroadcastSocket);
+	}
+
+	bool HasConnection() const
+	{
+		return BroadcastSocket != nullptr;
+	}
+
+	void Tick(const float InDeltaTime)
+	{
+		if (BroadcastSocket)
+		{
+			uint32 PendingDataSize = 0;
+			while (BroadcastSocket->HasPendingData(PendingDataSize) && PendingDataSize > 0)
+			{
+				TArray<uint8, TInlineAllocator<256>> ReceivedData;
+				ReceivedData.SetNumZeroed(PendingDataSize + 1);
+
+				int32 BytesRead = 0;
+				if (BroadcastSocket->Recv(ReceivedData.GetData(), ReceivedData.Num(), BytesRead) && BytesRead > 0)
+				{
+					ReceivedData.Last() = 0; // Ensure null terminator
+
+					FPythonScriptRemoteExecutionMessage Message;
+					if (PythonScriptRemoteExecutionUtil::ParseMessageData((char*)ReceivedData.GetData(), Message))
+					{
+						HandleMessage(Message);
+					}
+				}
+			}
+		}
+	}
+
+private:
+	void HandleMessage(const FPythonScriptRemoteExecutionMessage& InMessage)
+	{
+		if (!InMessage.PassesReceiveFilter(RemoteExecution->NodeId))
+		{
+			return;
+		}
+
+		if (InMessage.Type.Equals(PythonScriptRemoteExecutionUtil::TypePing))
+		{
+			BroadcastPongMessage(InMessage.Source);
+			return;
+		}
+
+		if (InMessage.Type.Equals(PythonScriptRemoteExecutionUtil::TypeOpenConnection))
+		{
+			HandleOpenConnectionMessage(InMessage);
+			return;
+		}
+
+		if (InMessage.Type.Equals(PythonScriptRemoteExecutionUtil::TypeCloseConnection))
+		{
+			HandleCloseConnectionMessage(InMessage);
+			return;
+		}
+
+		UE_LOG(LogPythonRemoteExecution, VeryVerbose, TEXT("Unhandled remote execution message type '%s'"), *InMessage.Type);
+	}
+
+	void HandleOpenConnectionMessage(const FPythonScriptRemoteExecutionMessage& InMessage)
+	{
+		const TSharedPtr<FJsonObject>* ConnectionData = nullptr;
+		if (InMessage.Data && InMessage.Data->TryGetObject(ConnectionData) && *ConnectionData)
+		{
+			FIPv4Address CommandIp;
+			{
+				FString CommandIpString;
+				if (!(*ConnectionData)->TryGetStringField(TEXT("command_ip"), CommandIpString))
+				{
+					UE_LOG(LogPythonRemoteExecution, Error, TEXT("Failed to find required data field 'command_ip' (string) for the 'open_connection' message!"));
+					return;
+				}
+				if (!FIPv4Address::Parse(CommandIpString, CommandIp))
+				{
+					UE_LOG(LogPythonRemoteExecution, Error, TEXT("Failed to parse command address '%s'!"), *CommandIpString);
+					return;
+				}
+			}
+
+			uint32 CommandPort = 0;
+			if (!(*ConnectionData)->TryGetNumberField(TEXT("command_port"), CommandPort))
+			{
+				UE_LOG(LogPythonRemoteExecution, Error, TEXT("Failed to find required data field 'command_port' (integer) for the 'open_connection' message!"));
+				return;
+			}
+
+			RemoteExecution->OpenCommandConnection(InMessage.Source, FIPv4Endpoint(CommandIp, (uint16)CommandPort));
+		}
+	}
+
+	void HandleCloseConnectionMessage(const FPythonScriptRemoteExecutionMessage& InMessage)
+	{
+		RemoteExecution->CloseCommandConnection(InMessage.Source);
+	}
+
+	void BroadcastMessage(const FPythonScriptRemoteExecutionMessage& InMessage)
+	{
+		FString MessageStr;
+		FText WriteErrorMsg;
+		if (!InMessage.ToJsonString(MessageStr, &WriteErrorMsg))
+		{
+			UE_LOG(LogPythonRemoteExecution, Error, TEXT("Failed to write remote execution message: %s"), *WriteErrorMsg.ToString());
+			return;
+		}
+		BroadcastMessage(TCHAR_TO_UTF8(*MessageStr));
+	}
+
+	void BroadcastMessage(const char* InMessage)
+	{
+		check(BroadcastSocket);
+
+		int32 BytesSent = 0;
+		BroadcastSocket->SendTo((uint8*)InMessage, FCStringAnsi::Strlen(InMessage), BytesSent, *BroadcastGroupAddr);
+	}
+
+	void BroadcastPongMessage(const FString& InDestinationNodeId)
+	{
+		BroadcastMessage(FPythonScriptRemoteExecutionMessage(PythonScriptRemoteExecutionUtil::TypePong, RemoteExecution->NodeId, InDestinationNodeId, MakeShared<FJsonValueObject>(NodeInfo)));
+	}
+
+	/** Owner remote execution instance */
+	FPythonScriptRemoteExecution* RemoteExecution;
+
+	/** The information about this node instance (used in "pong" messages) */
+	TSharedPtr<FJsonObject> NodeInfo;
+
+	/** The address of the UDP multicast group used for broadcast */
+	TSharedPtr<FInternetAddr> BroadcastGroupAddr;
+
+	/** The UDP multicast socket used for broadcast */
+	FSocket* BroadcastSocket;
+};
+
+class FPythonScriptRemoteExecutionCommandConnection
+{
+public:
+	FPythonScriptRemoteExecutionCommandConnection(FPythonScriptRemoteExecution* InRemoteExecution, const FString& InRemoteNodeId, const FIPv4Endpoint& InCommandEndpoint)
+		: RemoteExecution(InRemoteExecution)
+		, RemoteNodeId(InRemoteNodeId)
+		, CommandEndpoint(InCommandEndpoint)
+		, CommandChannelSocket(nullptr)
+	{
+		check(RemoteExecution);
+
+		const UPythonScriptPluginSettings* PythonScriptPluginSettings = GetDefault<UPythonScriptPluginSettings>();
+
+		CommandChannelSocket = FTcpSocketBuilder(TEXT("PythonRemoteExecutionCommandSocket"))
+			.WithSendBufferSize(PythonScriptPluginSettings->RemoteExecutionSendBufferSizeBytes)
+			.WithReceiveBufferSize(PythonScriptPluginSettings->RemoteExecutionReceiveBufferSizeBytes);
+
+		if (!CommandChannelSocket->Connect(*CommandEndpoint.ToInternetAddr()))
+		{
+			PythonScriptRemoteExecutionUtil::DestroySocket(RemoteExecution->SocketSubsystem, CommandChannelSocket);
+		}
+	}
+
+	~FPythonScriptRemoteExecutionCommandConnection()
+	{
+		PythonScriptRemoteExecutionUtil::DestroySocket(RemoteExecution->SocketSubsystem, CommandChannelSocket);
+	}
+
+	bool HasConnection() const
+	{
+		return CommandChannelSocket != nullptr;
+	}
+
+	bool IsConnectedTo(const FString& InRemoteNodeId)
+	{
+		return RemoteNodeId.Equals(InRemoteNodeId);
+	}
+
+	bool IsConnectedTo(const FIPv4Endpoint& InCommandEndpoint)
+	{
+		return CommandEndpoint == InCommandEndpoint;
+	}
+
+	bool IsConnectedTo(const FString& InRemoteNodeId, const FIPv4Endpoint& InCommandEndpoint)
+	{
+		return IsConnectedTo(InRemoteNodeId)
+			&& IsConnectedTo(InCommandEndpoint);
+	}
+
+	void Tick(const float InDeltaTime)
+	{
+		if (CommandChannelSocket)
+		{
+			uint32 PendingDataSize = 0;
+			while (CommandChannelSocket->HasPendingData(PendingDataSize) && PendingDataSize > 0)
+			{
+				TArray<uint8, TInlineAllocator<256>> ReceivedData;
+				ReceivedData.SetNumZeroed(PendingDataSize + 1);
+
+				int32 BytesRead = 0;
+				if (CommandChannelSocket->Recv(ReceivedData.GetData(), ReceivedData.Num(), BytesRead) && BytesRead > 0)
+				{
+					ReceivedData.Last() = 0; // Ensure null terminator
+
+					FPythonScriptRemoteExecutionMessage Message;
+					if (PythonScriptRemoteExecutionUtil::ParseMessageData((char*)ReceivedData.GetData(), Message))
+					{
+						HandleMessage(Message);
+					}
+				}
+			}
+		}
+	}
+
+private:
+	void HandleMessage(const FPythonScriptRemoteExecutionMessage& InMessage)
+	{
+		if (!InMessage.PassesReceiveFilter(RemoteExecution->NodeId))
+		{
+			return;
+		}
+
+		if (InMessage.Type.Equals(PythonScriptRemoteExecutionUtil::TypeCommand))
+		{
+			HandleCommandMessage(InMessage);
+			return;
+		}
+
+		UE_LOG(LogPythonRemoteExecution, VeryVerbose, TEXT("Unhandled remote execution message type '%s'"), *InMessage.Type);
+	}
+
+	void HandleCommandMessage(const FPythonScriptRemoteExecutionMessage& InMessage)
+	{
+		const TSharedPtr<FJsonObject>* CommandData = nullptr;
+		if (InMessage.Data && InMessage.Data->TryGetObject(CommandData) && *CommandData)
+		{
+			FPythonCommandEx PythonCommand;
+
+			if (!(*CommandData)->TryGetStringField(TEXT("command"), PythonCommand.Command))
+			{
+				UE_LOG(LogPythonRemoteExecution, Error, TEXT("Failed to find required data field 'command' (string) for the 'command' message!"));
+				return;
+			}
+
+			{
+				bool bUnattended = false;
+				if (!(*CommandData)->TryGetBoolField(TEXT("unattended"), bUnattended))
+				{
+					UE_LOG(LogPythonRemoteExecution, Error, TEXT("Failed to find required data field 'unattended' (bool) for the 'command' message!"));
+					return;
+				}
+				PythonCommand.Flags |= bUnattended ? EPythonCommandFlags::Unattended : EPythonCommandFlags::None;
+			}
+
+			{
+				FString ExecModeString;
+				if (!(*CommandData)->TryGetStringField(TEXT("exec_mode"), ExecModeString))
+				{
+					UE_LOG(LogPythonRemoteExecution, Error, TEXT("Failed to find required data field 'exec_mode' (string) for the 'command' message!"));
+					return;
+				}
+				if (!LexTryParseString(PythonCommand.ExecutionMode, *ExecModeString))
+				{
+					UE_LOG(LogPythonRemoteExecution, Error, TEXT("Failed to find parse execution mode '%s'!"), *ExecModeString);
+					return;
+				}
+			}
+
+			{
+				const bool bCommandSuccess = RemoteExecution->PythonScriptPlugin->ExecPythonCommandEx(PythonCommand);
+				SendCommandResultMessage(InMessage.Source, bCommandSuccess, PythonCommand.Command, PythonCommand.CommandResult, PythonCommand.LogOutput);
+			}
+		}
+	}
+
+	void SendMessage(const FPythonScriptRemoteExecutionMessage& InMessage)
+	{
+		FString MessageStr;
+		FText WriteErrorMsg;
+		if (!InMessage.ToJsonString(MessageStr, &WriteErrorMsg))
+		{
+			UE_LOG(LogPythonRemoteExecution, Error, TEXT("Failed to write remote execution message: %s"), *WriteErrorMsg.ToString());
+			return;
+		}
+		SendMessage(TCHAR_TO_UTF8(*MessageStr));
+	}
+
+	void SendMessage(const char* InMessage)
+	{
+		check(CommandChannelSocket);
+
+		int32 BytesSent = 0;
+		CommandChannelSocket->Send((uint8*)InMessage, FCStringAnsi::Strlen(InMessage), BytesSent);
+	}
+
+	void SendCommandResultMessage(const FString& InDestinationNodeId, const bool bSuccess, const FString& InCommand, const FString& InResult, const TArray<FPythonLogOutputEntry>& InLogOutput)
+	{
+		TArray<TSharedPtr<FJsonValue>> JsonLogOutput;
+		for (const FPythonLogOutputEntry& LogEntry : InLogOutput)
+		{
+			TSharedRef<FJsonObject> JsonLogEntry = MakeShared<FJsonObject>();
+			JsonLogEntry->SetStringField(TEXT("type"), LexToString(LogEntry.Type));
+			JsonLogEntry->SetStringField(TEXT("output"), LogEntry.Output);
+			JsonLogOutput.Add(MakeShared<FJsonValueObject>(JsonLogEntry));
+		}
+
+		TSharedRef<FJsonObject> ResultInfo = MakeShared<FJsonObject>();
+		ResultInfo->SetBoolField(TEXT("success"), bSuccess);
+		ResultInfo->SetStringField(TEXT("command"), InCommand);
+		ResultInfo->SetStringField(TEXT("result"), InResult);
+		ResultInfo->SetArrayField(TEXT("output"), JsonLogOutput);
+
+		SendMessage(FPythonScriptRemoteExecutionMessage(PythonScriptRemoteExecutionUtil::TypeCommandResult, RemoteExecution->NodeId, InDestinationNodeId, MakeShared<FJsonValueObject>(ResultInfo)));
+	}
+
+	/** Owner remote execution instance */
+	FPythonScriptRemoteExecution* RemoteExecution;
+
+	/** The ID of the remote node that we connected to */
+	FString RemoteNodeId;
+
+	/** Endpoint that this command channel should be connected to */
+	FIPv4Endpoint CommandEndpoint;
+
+	/** The TCP socket used to handle command requests */
+	FSocket* CommandChannelSocket;
+};
+
+FPythonScriptRemoteExecution::FPythonScriptRemoteExecution(IPythonScriptPlugin* InPythonScriptPlugin)
+	: PythonScriptPlugin(InPythonScriptPlugin)
+	, SocketSubsystem(ISocketSubsystem::Get())
+	, NodeId(FApp::GetInstanceId().ToString())
+{
+	check(PythonScriptPlugin);
+	check(SocketSubsystem);
+
+	if (GetDefault<UPythonScriptPluginSettings>()->bRemoteExecution)
+	{
+		Start();
+	}
+}
+
+FPythonScriptRemoteExecution::~FPythonScriptRemoteExecution()
+{
+	Stop();
+}
+
+bool FPythonScriptRemoteExecution::Start()
+{
+	Stop();
+
+	BroadcastConnection = MakeUnique<FPythonScriptRemoteExecutionBroadcastConnection>(this);
+	if (!BroadcastConnection->HasConnection())
+	{
+		UE_LOG(LogPythonRemoteExecution, Error, TEXT("Failed to open remote execution broadcast channel! Check your settings and the output log for more information."));
+		BroadcastConnection.Reset();
+	}
+
+	return BroadcastConnection.IsValid();
+}
+
+void FPythonScriptRemoteExecution::Stop()
+{
+	CommandConnection.Reset();
+	BroadcastConnection.Reset();
+}
+
+void FPythonScriptRemoteExecution::SyncToSettings()
+{
+	Stop();
+
+	if (GetDefault<UPythonScriptPluginSettings>()->bRemoteExecution)
+	{
+		Start();
+	}
+}
+
+void FPythonScriptRemoteExecution::Tick(const float InDeltaTime)
+{
+	if (CommandConnection)
+	{
+		CommandConnection->Tick(InDeltaTime);
+	}
+
+	if (BroadcastConnection)
+	{
+		BroadcastConnection->Tick(InDeltaTime);
+	}
+}
+
+void FPythonScriptRemoteExecution::OpenCommandConnection(const FString& InRemoteNodeId, const FIPv4Endpoint& InCommandEndpoint)
+{
+	if (!CommandConnection || !CommandConnection->IsConnectedTo(InRemoteNodeId, InCommandEndpoint))
+	{
+		CommandConnection = MakeUnique<FPythonScriptRemoteExecutionCommandConnection>(this, InRemoteNodeId, InCommandEndpoint);
+		if (!CommandConnection->HasConnection())
+		{
+			UE_LOG(LogPythonRemoteExecution, Error, TEXT("Failed to open remote execution command channel to '%s' (%s)!"), *InRemoteNodeId, *InCommandEndpoint.ToString());
+			CommandConnection.Reset();
+		}
+	}
+}
+
+void FPythonScriptRemoteExecution::CloseCommandConnection(const FString& InRemoteNodeId)
+{
+	if (CommandConnection && CommandConnection->IsConnectedTo(InRemoteNodeId))
+	{
+		CommandConnection.Reset();
+	}
+}
+
+namespace PythonScriptRemoteExecutionUtil
+{
+	bool SerializeJsonObject(const TSharedRef<FJsonObject>& InJsonObject, FString& OutJsonString, FText* OutError)
+	{
+		TSharedRef<TJsonWriter<TCHAR, TCondensedJsonPrintPolicy<TCHAR>>> JsonWriter = TJsonWriterFactory<TCHAR, TCondensedJsonPrintPolicy<TCHAR>>::Create(&OutJsonString);
+		const bool bDidSerialize = FJsonSerializer::Serialize(InJsonObject, JsonWriter);
+		JsonWriter->Close();
+
+		if (!bDidSerialize)
+		{
+			if (OutError)
+			{
+				*OutError = LOCTEXT("SerializeJsonObject.Error.FailedToSerialize", "Failed to serialize JSON object!");
+			}
+			return false;
+		}
+
+		return true;
+	}
+
+	bool DeserializeJsonObject(const FString& InJsonString, TSharedPtr<FJsonObject>& OutJsonObject, FText* OutError)
+	{
+		TSharedRef<TJsonReader<TCHAR>> JsonWriter = TJsonReaderFactory<TCHAR>::Create(InJsonString);
+		if (!FJsonSerializer::Deserialize(JsonWriter, OutJsonObject))
+		{
+			if (OutError)
+			{
+				*OutError = FText::Format(LOCTEXT("DeserializeJsonObject.Error.FailedToDeserialize", "Failed to deserialize JSON object! {0}"), FText::AsCultureInvariant(JsonWriter->GetErrorMessage()));
+			}
+			return false;
+		}
+
+		return true;
+	}
+
+	bool ParseMessageData(const char* InMessageData, FPythonScriptRemoteExecutionMessage& OutMessage)
+	{
+		return ParseMessageData(UTF8_TO_TCHAR(InMessageData), OutMessage);
+	}
+
+	bool ParseMessageData(const TCHAR* InMessageData, FPythonScriptRemoteExecutionMessage& OutMessage)
+	{
+		FText ReadErrorMsg;
+		if (!OutMessage.FromJsonString(InMessageData, &ReadErrorMsg))
+		{
+			UE_LOG(LogPythonRemoteExecution, Error, TEXT("Failed to read remote execution message '%s': %s"), InMessageData, *ReadErrorMsg.ToString());
+			return false;
+		}
+		return true;
+	}
+
+	bool WriteMessageData(const FPythonScriptRemoteExecutionMessage& InMessage, FString& OutMessageData)
+	{
+		FText WriteErrorMsg;
+		if (!InMessage.ToJsonString(OutMessageData, &WriteErrorMsg))
+		{
+			UE_LOG(LogPythonRemoteExecution, Error, TEXT("Failed to write remote execution message: %s"), *WriteErrorMsg.ToString());
+			return false;
+		}
+		return true;
+	}
+
+	void DestroySocket(ISocketSubsystem* InSocketSubsystem, FSocket*& InOutSocket)
+	{
+		check(InSocketSubsystem);
+
+		if (InOutSocket)
+		{
+			if (InOutSocket->GetConnectionState() == SCS_Connected)
+			{
+				InOutSocket->Close();
+			}
+
+			InSocketSubsystem->DestroySocket(InOutSocket);
+			InOutSocket = nullptr;
+		}
+	}
+}
+
+#endif	// WITH_PYTHON
+
+#undef LOCTEXT_NAMESPACE
diff --git a/Engine/Plugins/Experimental/PythonScriptPlugin/Source/PythonScriptPlugin/Private/PythonScriptRemoteExecution.h b/Engine/Plugins/Experimental/PythonScriptPlugin/Source/PythonScriptPlugin/Private/PythonScriptRemoteExecution.h
new file mode 100644
--- /dev/null
+++ b/Engine/Plugins/Experimental/PythonScriptPlugin/Source/PythonScriptPlugin/Private/PythonScriptRemoteExecution.h
@@ -0,0 +1,61 @@
+// Copyright 1998-2019 Epic Games, Inc. All Rights Reserved.
+
+#pragma once
+
+#include "IncludePython.h"
+#include "Templates/UniquePtr.h"
+
+struct FIPv4Endpoint;
+
+class ISocketSubsystem;
+class IPythonScriptPlugin;
+class FPythonScriptRemoteExecutionCommandConnection;
+class FPythonScriptRemoteExecutionBroadcastConnection;
+
+#if WITH_PYTHON
+
+class FPythonScriptRemoteExecution
+{
+	friend class FPythonScriptRemoteExecutionCommandConnection;
+	friend class FPythonScriptRemoteExecutionBroadcastConnection;
+
+public:
+	explicit FPythonScriptRemoteExecution(IPythonScriptPlugin* InPythonScriptPlugin);
+	~FPythonScriptRemoteExecution();
+
+	/** Start remote execution based on the current settings (also called during construction if remote execution should be active) */
+	bool Start();
+
+	/** Stop remote execution (also called during destruction if remote execution is active) */
+	void Stop();
+
+	/** Sync the remote execution environment to the current settings, starting or stopping it as required */
+	void SyncToSettings();
+
+	/** Tick, processing and sending messages as required */
+	void Tick(const float InDeltaTime);
+
+private:
+	/** Open a command connection to the given remote node, at the given endpoint */
+	void OpenCommandConnection(const FString& InRemoteNodeId, const FIPv4Endpoint& InCommandEndpoint);
+
+	/** Close any existing command connection to the given remote node */
+	void CloseCommandConnection(const FString& InRemoteNodeId);
+
+	/** The Python plugin that owns this instance */
+	IPythonScriptPlugin* PythonScriptPlugin;
+
+	/** Cached socket subsystem pointer */
+	ISocketSubsystem* SocketSubsystem;
+
+	/** The ID of this remote execution node */
+	FString NodeId;
+
+	/** Connection handling TCP commands */
+	TUniquePtr<FPythonScriptRemoteExecutionCommandConnection> CommandConnection;
+
+	/** Connection handling UDP broadcast */
+	TUniquePtr<FPythonScriptRemoteExecutionBroadcastConnection> BroadcastConnection;
+};
+
+#endif	// WITH_PYTHON
diff --git a/Engine/Plugins/Experimental/PythonScriptPlugin/Source/PythonScriptPlugin/PythonScriptPlugin.Build.cs b/Engine/Plugins/Experimental/PythonScriptPlugin/Source/PythonScriptPlugin/PythonScriptPlugin.Build.cs
--- a/Engine/Plugins/Experimental/PythonScriptPlugin/Source/PythonScriptPlugin/PythonScriptPlugin.Build.cs
+++ b/Engine/Plugins/Experimental/PythonScriptPlugin/Source/PythonScriptPlugin/PythonScriptPlugin.Build.cs
@@ -23,6 +23,9 @@ namespace UnrealBuildTool.Rules
 					"Slate",
 					"SlateCore",
 					"InputCore",
+					"Sockets",
+					"Networking",
+					"Json",
 				}
 			);
 
