Fixed build error for platforms without a Python SDK available

diff --git a/Engine/Plugins/Experimental/PythonScriptPlugin/Source/PythonScriptPlugin/Private/PythonScriptPluginSettings.cpp b/Engine/Plugins/Experimental/PythonScriptPlugin/Source/PythonScriptPlugin/Private/PythonScriptPluginSettings.cpp
--- a/Engine/Plugins/Experimental/PythonScriptPlugin/Source/PythonScriptPlugin/Private/PythonScriptPluginSettings.cpp
+++ b/Engine/Plugins/Experimental/PythonScriptPlugin/Source/PythonScriptPlugin/Private/PythonScriptPluginSettings.cpp
@@ -43,6 +43,7 @@ void UPythonScriptPluginSettings::PostEditChangeProperty(struct FPropertyChanged
 {
 	Super::PostEditChangeProperty(PropertyChangedEvent);
 
+#if WITH_PYTHON
 	if (PropertyChangedEvent.MemberProperty)
 	{
 		if (PropertyChangedEvent.MemberProperty->GetFName() == GET_MEMBER_NAME_CHECKED(UPythonScriptPluginSettings, bRemoteExecution) ||
@@ -56,6 +57,7 @@ void UPythonScriptPluginSettings::PostEditChangeProperty(struct FPropertyChanged
 			FPythonScriptPlugin::Get()->SyncRemoteExecutionToSettings();
 		}
 	}
+#endif	// WITH_PYTHON
 }
 
 FText UPythonScriptPluginSettings::GetSectionText() const
