Fixed Python doc generation

#rnx

diff --git a/Engine/Plugins/Experimental/PythonScriptPlugin/Source/PythonScriptPlugin/Private/PyOnlineDocsWriter.cpp b/Engine/Plugins/Experimental/PythonScriptPlugin/Source/PythonScriptPlugin/Private/PyOnlineDocsWriter.cpp
--- a/Engine/Plugins/Experimental/PythonScriptPlugin/Source/PythonScriptPlugin/Private/PyOnlineDocsWriter.cpp
+++ b/Engine/Plugins/Experimental/PythonScriptPlugin/Source/PythonScriptPlugin/Private/PyOnlineDocsWriter.cpp
@@ -329,7 +329,7 @@ void FPyOnlineDocsWriter::GenerateFiles(const FString& InPythonStubPath)
 		//   https://pip.pypa.io/en/stable/user_guide/#using-pip-from-your-program
 
 		FString PyCommandStr;
-		FString PythonPath = FPaths::ConvertRelativePathToFull(FPaths::EngineSourceDir()) / TEXT("ThirdParty/Python/Win64/python.exe");
+		FString PythonPath = FPaths::ConvertRelativePathToFull(FPaths::EngineDir()) / TEXT("Binaries/ThirdParty/Python/Win64/python.exe");
 
 		PyCommandStr += TEXT(
 			"import sys\n"
