Prevent a crash when a Python generated class pending reinstancing is GC'd

We now ARO these pending types to stop them being GC'd, and added unreal.flush_generated_type_reinstancing() to allow you to force the reinstancing to happen at a certain point in your script (rather than when it ends, or on the next Tick).

diff --git a/Engine/Plugins/Experimental/PythonScriptPlugin/Source/PythonScriptPlugin/Private/PyCore.cpp b/Engine/Plugins/Experimental/PythonScriptPlugin/Source/PythonScriptPlugin/Private/PyCore.cpp
--- a/Engine/Plugins/Experimental/PythonScriptPlugin/Source/PythonScriptPlugin/Private/PyCore.cpp
+++ b/Engine/Plugins/Experimental/PythonScriptPlugin/Source/PythonScriptPlugin/Private/PyCore.cpp
@@ -1528,6 +1528,13 @@ PyObject* GenerateEnum(PyObject* InSelf, PyObject* InArgs)
 	Py_RETURN_NONE;
 }
 
+PyObject* FlushGeneratedTypeReinstancing(PyObject* InSelf)
+{
+	FPyWrapperTypeReinstancer::Get().ProcessPending();
+
+	Py_RETURN_NONE;
+}
+
 PyObject* GetTypeFromClass(PyObject* InSelf, PyObject* InArgs)
 {
 	PyObject* PyObj = nullptr;
@@ -1665,6 +1672,7 @@ PyMethodDef PyCoreMethods[] = {
 	{ "generate_class", PyCFunctionCast(&GenerateClass), METH_VARARGS, "x.generate_class(type) -> None -- generate an Unreal class for the given Python type" },
 	{ "generate_struct", PyCFunctionCast(&GenerateStruct), METH_VARARGS, "x.generate_struct(type) -> None -- generate an Unreal struct for the given Python type" },
 	{ "generate_enum", PyCFunctionCast(&GenerateEnum), METH_VARARGS, "x.generate_enum(type) -> None -- generate an Unreal enum for the given Python type" },
+	{ "flush_generated_type_reinstancing", PyCFunctionCast(&FlushGeneratedTypeReinstancing), METH_NOARGS, "x.flush_generated_type_reinstancing() -> None -- flush any pending reinstancing requests for Python generated types" },
 	{ "get_type_from_class", PyCFunctionCast(&GetTypeFromClass), METH_VARARGS, "x.get_type_from_class(class) -> type -- get the best matching Python type for the given Unreal class" },
 	{ "get_type_from_struct", PyCFunctionCast(&GetTypeFromStruct), METH_VARARGS, "x.get_type_from_struct(struct) -> type -- get the best matching Python type for the given Unreal struct" },
 	{ "get_type_from_enum", PyCFunctionCast(&GetTypeFromEnum), METH_VARARGS, "x.get_type_from_enum(enum) -> type -- get the best matching Python type for the given Unreal enum" },
diff --git a/Engine/Plugins/Experimental/PythonScriptPlugin/Source/PythonScriptPlugin/Private/PyReferenceCollector.cpp b/Engine/Plugins/Experimental/PythonScriptPlugin/Source/PythonScriptPlugin/Private/PyReferenceCollector.cpp
--- a/Engine/Plugins/Experimental/PythonScriptPlugin/Source/PythonScriptPlugin/Private/PyReferenceCollector.cpp
+++ b/Engine/Plugins/Experimental/PythonScriptPlugin/Source/PythonScriptPlugin/Private/PyReferenceCollector.cpp
@@ -1,6 +1,7 @@
 // Copyright 1998-2019 Epic Games, Inc. All Rights Reserved.
 
 #include "PyReferenceCollector.h"
+#include "PyWrapperTypeRegistry.h"
 #include "PyWrapperBase.h"
 #include "PyWrapperDelegate.h"
 #include "UObject/UnrealType.h"
@@ -35,6 +36,8 @@ void FPyReferenceCollector::AddReferencedObjects(FReferenceCollector& InCollecto
 			PythonWrappedInstanceMetaData->AddReferencedObjects(PythonWrappedInstance, InCollector);
 		}
 	}
+
+	FPyWrapperTypeReinstancer::Get().AddReferencedObjects(InCollector);
 }
 
 void FPyReferenceCollector::PurgeUnrealObjectReferences(const UObject* InObject, const bool bIncludeInnerObjects)
diff --git a/Engine/Plugins/Experimental/PythonScriptPlugin/Source/PythonScriptPlugin/Private/PyWrapperTypeRegistry.cpp b/Engine/Plugins/Experimental/PythonScriptPlugin/Source/PythonScriptPlugin/Private/PyWrapperTypeRegistry.cpp
--- a/Engine/Plugins/Experimental/PythonScriptPlugin/Source/PythonScriptPlugin/Private/PyWrapperTypeRegistry.cpp
+++ b/Engine/Plugins/Experimental/PythonScriptPlugin/Source/PythonScriptPlugin/Private/PyWrapperTypeRegistry.cpp
@@ -370,8 +370,11 @@ void FPyWrapperTypeReinstancer::ProcessPending()
 	{
 		for (const auto& ClassToReinstancePair : ClassesToReinstance)
 		{
-			// Assume the classes have changed
-			FCoreUObjectDelegates::RegisterClassForHotReloadReinstancingDelegate.Broadcast(ClassToReinstancePair.Key, ClassToReinstancePair.Value, EHotReloadedClassFlags::Changed);
+			if (ClassToReinstancePair.Key && ClassToReinstancePair.Value)
+			{
+				// Assume the classes have changed
+				FCoreUObjectDelegates::RegisterClassForHotReloadReinstancingDelegate.Broadcast(ClassToReinstancePair.Key, ClassToReinstancePair.Value, EHotReloadedClassFlags::Changed);
+			}
 		}
 		FCoreUObjectDelegates::ReinstanceHotReloadedClassesDelegate.Broadcast();
 
@@ -379,6 +382,22 @@ void FPyWrapperTypeReinstancer::ProcessPending()
 	}
 
 	// todo: need support for re-instancing structs
+	StructsToReinstance.Reset();
+}
+
+void FPyWrapperTypeReinstancer::AddReferencedObjects(FReferenceCollector& InCollector)
+{
+	for (auto& ClassToReinstancePair : ClassesToReinstance)
+	{
+		InCollector.AddReferencedObject(ClassToReinstancePair.Key);
+		InCollector.AddReferencedObject(ClassToReinstancePair.Value);
+	}
+
+	for (auto& StructToReinstancePair : StructsToReinstance)
+	{
+		InCollector.AddReferencedObject(StructToReinstancePair.Key);
+		InCollector.AddReferencedObject(StructToReinstancePair.Value);
+	}
 }
 
 
diff --git a/Engine/Plugins/Experimental/PythonScriptPlugin/Source/PythonScriptPlugin/Private/PyWrapperTypeRegistry.h b/Engine/Plugins/Experimental/PythonScriptPlugin/Source/PythonScriptPlugin/Private/PyWrapperTypeRegistry.h
--- a/Engine/Plugins/Experimental/PythonScriptPlugin/Source/PythonScriptPlugin/Private/PyWrapperTypeRegistry.h
+++ b/Engine/Plugins/Experimental/PythonScriptPlugin/Source/PythonScriptPlugin/Private/PyWrapperTypeRegistry.h
@@ -31,6 +31,7 @@ class UPythonGeneratedStruct;
 class IPyWrapperInlineStructFactory;
 
 class FPyFileWriter;
+class FReferenceCollector;
 
 /** Type conversion for TPyWrapperTypeFactory */
 template <typename UnrealType, typename KeyType>
@@ -300,6 +301,9 @@ public:
 	/** Process any pending re-instance requests */
 	void ProcessPending();
 
+	/** Collect any referenced objects */
+	void AddReferencedObjects(FReferenceCollector& InCollector);
+
 private:
 	/** Pending pairs of classes that to be re-instanced */
 	TArray<TPair<UPythonGeneratedClass*, UPythonGeneratedClass*>> ClassesToReinstance;
