init_unreal.py files are now run using the console scope to let them set-up the console environment in the editor

diff --git a/Engine/Plugins/Experimental/PythonScriptPlugin/Source/PythonScriptPlugin/Private/PythonScriptPlugin.cpp b/Engine/Plugins/Experimental/PythonScriptPlugin/Source/PythonScriptPlugin/Private/PythonScriptPlugin.cpp
--- a/Engine/Plugins/Experimental/PythonScriptPlugin/Source/PythonScriptPlugin/Private/PythonScriptPlugin.cpp
+++ b/Engine/Plugins/Experimental/PythonScriptPlugin/Source/PythonScriptPlugin/Private/PythonScriptPlugin.cpp
@@ -743,8 +743,11 @@ void FPythonScriptPlugin::Tick(const float InDeltaTime)
 			const FString PotentialFilePath = PySysPath / TEXT("init_unreal.py");
 			if (FPaths::FileExists(PotentialFilePath))
 			{
-				// Quote the path in case it contains spaces so that the token parsing will work as expected
-				ExecPythonCommand(*FString::Printf(TEXT("\"%s\""), *PotentialFilePath));
+				// Execute these files in the "public" scope, as if their contents had been run directly in the console
+				// This allows them to be used to set-up an editor environment for the console
+				FPythonCommandEx InitUnrealPythonCommand;
+				InitUnrealPythonCommand.FileExecutionScope = EPythonFileExecutionScope::Public;
+				RunFile(*PotentialFilePath, *InitUnrealPythonCommand.Command, InitUnrealPythonCommand);
 			}
 		}
 		for (const FString& StartupScript : GetDefault<UPythonScriptPluginSettings>()->StartupScripts)
@@ -960,8 +963,13 @@ bool FPythonScriptPlugin::RunFile(const TCHAR* InFile, const TCHAR* InArgs, FPyt
 		FPyScopedGIL GIL;
 		TGuardValue<bool> UnattendedScriptGuard(GIsRunningUnattendedScript, GIsRunningUnattendedScript || EnumHasAnyFlags(InOutPythonCommand.Flags, EPythonCommandFlags::Unattended));
 
-		FPyObjectPtr PyFileGlobalDict = FPyObjectPtr::StealReference(PyDict_Copy(PyDefaultGlobalDict));
-		FPyObjectPtr PyFileLocalDict = PyFileGlobalDict;
+		FPyObjectPtr PyFileGlobalDict = PyConsoleGlobalDict;
+		FPyObjectPtr PyFileLocalDict = PyConsoleLocalDict;
+		if (InOutPythonCommand.FileExecutionScope == EPythonFileExecutionScope::Private)
+		{
+			PyFileGlobalDict = FPyObjectPtr::StealReference(PyDict_Copy(PyDefaultGlobalDict));
+			PyFileLocalDict = PyFileGlobalDict;
+		}
 		{
 			FPyObjectPtr PyResolvedFilePath;
 			if (PyConversion::Pythonize(ResolvedFilePath, PyResolvedFilePath.Get(), PyConversion::ESetErrorState::No))
@@ -980,6 +988,8 @@ bool FPythonScriptPlugin::RunFile(const TCHAR* InFile, const TCHAR* InArgs, FPyt
 			PyCore::GetPythonLogCapture().Remove(LogCaptureHandle);
 		}
 
+		PyDict_DelItemString(PyFileGlobalDict, "__file__");
+
 		if (PyResult)
 		{
 			InOutPythonCommand.CommandResult = PyUtil::PyObjectToUEStringRepr(PyResult);
diff --git a/Engine/Plugins/Experimental/PythonScriptPlugin/Source/PythonScriptPlugin/Public/IPythonScriptPlugin.h b/Engine/Plugins/Experimental/PythonScriptPlugin/Source/PythonScriptPlugin/Public/IPythonScriptPlugin.h
--- a/Engine/Plugins/Experimental/PythonScriptPlugin/Source/PythonScriptPlugin/Public/IPythonScriptPlugin.h
+++ b/Engine/Plugins/Experimental/PythonScriptPlugin/Source/PythonScriptPlugin/Public/IPythonScriptPlugin.h
@@ -80,6 +80,15 @@ enum class EPythonCommandExecutionMode : uint8
 	EvaluateStatement,
 };
 
+/** Controls the scope used when executing Python files. */
+enum class EPythonFileExecutionScope : uint8
+{
+	/** Execute the Python file with its own unique locals and globals dict to isolate any changes it makes to the environment (like imports). */
+	Private,
+	/** Execute the Python file with the shared locals and globals dict as used by the console, so that executing it behaves as if you'd ran the file contents directly in the console. */
+	Public,
+};
+
 /** Log output entry captured from Python. */
 struct FPythonLogOutputEntry
 {
@@ -99,6 +108,9 @@ struct FPythonCommandEx
 	/** Controls the mode used to execute the command. */
 	EPythonCommandExecutionMode ExecutionMode = EPythonCommandExecutionMode::ExecuteFile;
 
+	/** Controls the scope used when executing Python files. */
+	EPythonFileExecutionScope FileExecutionScope = EPythonFileExecutionScope::Private;
+
 	/** The command to run. This may be literal Python code, or a file (with optional arguments) to run. */
 	FString Command;
 
