Fixed name conflict when re-instancing Python generated structs

diff --git a/Engine/Plugins/Experimental/PythonScriptPlugin/Source/PythonScriptPlugin/Private/PyWrapperStruct.cpp b/Engine/Plugins/Experimental/PythonScriptPlugin/Source/PythonScriptPlugin/Private/PyWrapperStruct.cpp
--- a/Engine/Plugins/Experimental/PythonScriptPlugin/Source/PythonScriptPlugin/Private/PyWrapperStruct.cpp
+++ b/Engine/Plugins/Experimental/PythonScriptPlugin/Source/PythonScriptPlugin/Private/PyWrapperStruct.cpp
@@ -1347,7 +1347,11 @@ public:
 			}
 		}
 
-		// We can no longer fail, set the correct name on the new struct
+		// We can no longer fail, so prepare the old struct for removal and set the correct name on the new struct
+		if (OldStruct)
+		{
+			PrepareOldStructForReinstancing();
+		}
 		NewStruct->Rename(*StructName, nullptr, REN_DontCreateRedirectors);
 
 		// Finalize the struct
