Ensure Python callable wrapper types remove the Native flag from any function signatures that they duplicate

Failing to do this leads to a Native instance inside a non-native instance, which causes GC confusion

diff --git a/Engine/Plugins/Experimental/PythonScriptPlugin/Source/PythonScriptPlugin/Private/PyWrapperTypeRegistry.cpp b/Engine/Plugins/Experimental/PythonScriptPlugin/Source/PythonScriptPlugin/Private/PyWrapperTypeRegistry.cpp
--- a/Engine/Plugins/Experimental/PythonScriptPlugin/Source/PythonScriptPlugin/Private/PyWrapperTypeRegistry.cpp
+++ b/Engine/Plugins/Experimental/PythonScriptPlugin/Source/PythonScriptPlugin/Private/PyWrapperTypeRegistry.cpp
@@ -1796,7 +1796,13 @@ PyTypeObject* FPyWrapperTypeRegistry::GenerateWrappedDelegateType(const UFunctio
 	UClass* PythonCallableForDelegateClass = nullptr;
 	{
 		PythonCallableForDelegateClass = NewObject<UClass>(GetPythonTypeContainer(), *FString::Printf(TEXT("%s__PythonCallable"), *DelegateBaseTypename), RF_Public);
-		UFunction* PythonCallableForDelegateFunc = (UFunction*)StaticDuplicateObject(InDelegateSignature, PythonCallableForDelegateClass, UPythonCallableForDelegate::GeneratedFuncName);
+		UFunction* PythonCallableForDelegateFunc = nullptr;
+		{
+			FObjectDuplicationParameters FuncDuplicationParams(const_cast<UFunction*>(InDelegateSignature), PythonCallableForDelegateClass);
+			FuncDuplicationParams.DestName = UPythonCallableForDelegate::GeneratedFuncName;
+			FuncDuplicationParams.InternalFlagMask &= ~EInternalObjectFlags::Native;
+			PythonCallableForDelegateFunc = CastChecked<UFunction>(StaticDuplicateObjectEx(FuncDuplicationParams));
+		}
 		PythonCallableForDelegateFunc->FunctionFlags = (PythonCallableForDelegateFunc->FunctionFlags | FUNC_Native) & ~(FUNC_Delegate | FUNC_MulticastDelegate);
 		PythonCallableForDelegateFunc->SetNativeFunc(&UPythonCallableForDelegate::CallPythonNative);
 		PythonCallableForDelegateFunc->StaticLink(true);
