Attempt to forcefully purge any Python generated Unreal types and their instances when shutting down the Python plugin

diff --git a/Engine/Plugins/Experimental/PythonScriptPlugin/Source/PythonScriptPlugin/Private/PyReferenceCollector.cpp b/Engine/Plugins/Experimental/PythonScriptPlugin/Source/PythonScriptPlugin/Private/PyReferenceCollector.cpp
--- a/Engine/Plugins/Experimental/PythonScriptPlugin/Source/PythonScriptPlugin/Private/PyReferenceCollector.cpp
+++ b/Engine/Plugins/Experimental/PythonScriptPlugin/Source/PythonScriptPlugin/Private/PyReferenceCollector.cpp
@@ -3,6 +3,9 @@
 #include "PyReferenceCollector.h"
 #include "PyWrapperTypeRegistry.h"
 #include "PyWrapperBase.h"
+#include "PyWrapperObject.h"
+#include "PyWrapperStruct.h"
+#include "PyWrapperEnum.h"
 #include "PyWrapperDelegate.h"
 #include "UObject/UnrealType.h"
 #include "UObject/UObjectHash.h"
@@ -55,13 +58,10 @@ void FPyReferenceCollector::PurgeUnrealObjectReferences(const TArrayView<const U
 
 		if (bIncludeInnerObjects)
 		{
-			TArray<UObject*> InnerObjects;
-			GetObjectsWithOuter(Object, InnerObjects, true);
-
-			for (const UObject* InnerObject : InnerObjects)
+			ForEachObjectWithOuter(Object, [&PurgingReferenceCollector](UObject* InnerObject)
 			{
 				PurgingReferenceCollector.AddObjectToPurge(InnerObject);
-			}
+			}, true);
 		}
 	}
 
@@ -71,6 +71,125 @@ void FPyReferenceCollector::PurgeUnrealObjectReferences(const TArrayView<const U
 	}
 }
 
+void FPyReferenceCollector::PurgeUnrealGeneratedTypes()
+{
+	auto RunPurgeUnrealGeneratedTypes = [this](const bool bLogFailure) -> bool
+	{
+		FPurgingReferenceCollector PurgingReferenceCollector;
+		TArray<FWeakObjectPtr> WeakReferencesToPurgedObjects;
+
+		auto FlagObjectForPurge = [&PurgingReferenceCollector, &WeakReferencesToPurgedObjects](UObject* InObject, const bool bMarkPendingKill)
+		{
+			if (!InObject->HasAnyInternalFlags(EInternalObjectFlags::Native) && !InObject->HasAnyFlags(RF_ClassDefaultObject))
+			{
+				if (InObject->IsRooted())
+				{
+					InObject->RemoveFromRoot();
+				}
+				InObject->ClearFlags(RF_Public | RF_Standalone);
+				if (bMarkPendingKill)
+				{
+					InObject->MarkPendingKill();
+				}
+
+				WeakReferencesToPurgedObjects.Add(InObject);
+			}
+
+			PurgingReferenceCollector.AddObjectToPurge(InObject);
+		};
+
+		// Clean-up Python generated class types and instances
+		// The class types are instances of UPythonGeneratedClass
+		{
+			ForEachObjectOfClass(UPythonGeneratedClass::StaticClass(), [&FlagObjectForPurge](UObject* InObject)
+			{
+				UPythonGeneratedClass* PythonGeneratedClass = CastChecked<UPythonGeneratedClass>(InObject);
+				bool bMarkClassPendingKill = true; // Mark types as PendingKill if they have no instances left (excluding their CDO)
+
+				ForEachObjectOfClass(PythonGeneratedClass, [&bMarkClassPendingKill, &FlagObjectForPurge](UObject* InInnerObject)
+				{
+					bMarkClassPendingKill &= InInnerObject->HasAnyFlags(RF_ClassDefaultObject);
+					FlagObjectForPurge(InInnerObject, /*bMarkPendingKill*/true);
+				}, false);
+
+				FlagObjectForPurge(PythonGeneratedClass, bMarkClassPendingKill);
+			}, false);
+		}
+
+		// Clean-up Python generated struct types
+		// The struct types are instances of UPythonGeneratedStruct
+		{
+			ForEachObjectOfClass(UPythonGeneratedStruct::StaticClass(), [&FlagObjectForPurge](UObject* InObject)
+			{
+				FlagObjectForPurge(InObject, /*bMarkPendingKill*/true);
+			}, false);
+		}
+
+		// Clean-up Python generated enum types
+		// The enum types are instances of UPythonGeneratedEnum
+		{
+			ForEachObjectOfClass(UPythonGeneratedEnum::StaticClass(), [&FlagObjectForPurge](UObject* InObject)
+			{
+				FlagObjectForPurge(InObject, /*bMarkPendingKill*/true);
+			}, false);
+		}
+
+		// Clean-up Python callable types and instances
+		// The callable types all derive directly from UPythonCallableForDelegate
+		{
+			TArray<UClass*> PythonCallableClasses;
+			GetDerivedClasses(UPythonCallableForDelegate::StaticClass(), PythonCallableClasses, true);
+
+			for (UClass* PythonCallableClass : PythonCallableClasses)
+			{
+				bool bMarkClassPendingKill = true; // Mark types as PendingKill if they have no instances left (excluding their CDO)
+
+				ForEachObjectOfClass(PythonCallableClass, [&bMarkClassPendingKill, &FlagObjectForPurge](UObject* InObject)
+				{
+					bMarkClassPendingKill &= InObject->HasAnyFlags(RF_ClassDefaultObject);
+					FlagObjectForPurge(InObject, /*bMarkPendingKill*/true);
+				}, false);
+
+				FlagObjectForPurge(PythonCallableClass, bMarkClassPendingKill);
+			}
+		}
+
+		if (PurgingReferenceCollector.HasObjectToPurge())
+		{
+			AddReferencedObjects(PurgingReferenceCollector);
+			CollectGarbage(GARBAGE_COLLECTION_KEEPFLAGS);
+
+			bool bHasLeftoverObjects = false;
+			for (const FWeakObjectPtr& WeakReferencesToPurgedObject : WeakReferencesToPurgedObjects)
+			{
+				if (UObject* FailedToPurgeObject = WeakReferencesToPurgedObject.Get())
+				{
+					bHasLeftoverObjects = true;
+					if (bLogFailure)
+					{
+						UE_LOG(LogPython, Warning, TEXT("Object '%s' failed to purge when requested by PurgeUnrealGeneratedTypes. This may lead to crashes!"), *FailedToPurgeObject->GetPathName());
+					}
+					else if (bHasLeftoverObjects)
+					{
+						break;
+					}
+				}
+			}
+			return bHasLeftoverObjects;
+		}
+
+		return false;
+	};
+
+	// We run two purge passes:
+	//  - Pass 1 will force purge any type instances, and any types that have no instances
+	//  - Pass 2 will purge any types that no longer have instances left after running pass 1
+	if (RunPurgeUnrealGeneratedTypes(/*bLogFailure*/false))
+	{
+		RunPurgeUnrealGeneratedTypes(/*bLogFailure*/true);
+	}
+}
+
 void FPyReferenceCollector::AddReferencedObjectsFromDelegate(FReferenceCollector& InCollector, FScriptDelegate& InDelegate)
 {
 	// Keep the delegate object alive if it's using a Python proxy instance
diff --git a/Engine/Plugins/Experimental/PythonScriptPlugin/Source/PythonScriptPlugin/Private/PyReferenceCollector.h b/Engine/Plugins/Experimental/PythonScriptPlugin/Source/PythonScriptPlugin/Private/PyReferenceCollector.h
--- a/Engine/Plugins/Experimental/PythonScriptPlugin/Source/PythonScriptPlugin/Private/PyReferenceCollector.h
+++ b/Engine/Plugins/Experimental/PythonScriptPlugin/Source/PythonScriptPlugin/Private/PyReferenceCollector.h
@@ -44,6 +44,9 @@ public:
 	/** Purge any Python references to the given UObject instances */
 	void PurgeUnrealObjectReferences(const TArrayView<const UObject*>& InObjects, const bool bIncludeInnerObjects);
 
+	/** Purge any Python references to any Unreal types generated from or for Python */
+	void PurgeUnrealGeneratedTypes();
+
 	//~ FGCObject interface
 	virtual void AddReferencedObjects(FReferenceCollector& InCollector) override;
 
diff --git a/Engine/Plugins/Experimental/PythonScriptPlugin/Source/PythonScriptPlugin/Private/PythonScriptPlugin.cpp b/Engine/Plugins/Experimental/PythonScriptPlugin/Source/PythonScriptPlugin/Private/PythonScriptPlugin.cpp
--- a/Engine/Plugins/Experimental/PythonScriptPlugin/Source/PythonScriptPlugin/Private/PythonScriptPlugin.cpp
+++ b/Engine/Plugins/Experimental/PythonScriptPlugin/Source/PythonScriptPlugin/Private/PythonScriptPlugin.cpp
@@ -668,6 +668,8 @@ void FPythonScriptPlugin::ShutdownPython()
 	FEditorSupportDelegates::PrepareToCleanseEditorObject.RemoveAll(this);
 #endif	// WITH_EDITOR
 
+	FPyReferenceCollector::Get().PurgeUnrealGeneratedTypes();
+
 	PyUnrealModule.Reset();
 	PyDefaultGlobalDict.Reset();
 	PyDefaultLocalDict.Reset();
