Disabling agressive Python type purging (and logging of lingering types) due to a GC crash

#rnx

diff --git a/Engine/Plugins/Experimental/PythonScriptPlugin/Source/PythonScriptPlugin/Private/PyReferenceCollector.cpp b/Engine/Plugins/Experimental/PythonScriptPlugin/Source/PythonScriptPlugin/Private/PyReferenceCollector.cpp
--- a/Engine/Plugins/Experimental/PythonScriptPlugin/Source/PythonScriptPlugin/Private/PyReferenceCollector.cpp
+++ b/Engine/Plugins/Experimental/PythonScriptPlugin/Source/PythonScriptPlugin/Private/PyReferenceCollector.cpp
@@ -104,7 +104,7 @@ void FPyReferenceCollector::PurgeUnrealGeneratedTypes()
 			ForEachObjectOfClass(UPythonGeneratedClass::StaticClass(), [&FlagObjectForPurge](UObject* InObject)
 			{
 				UPythonGeneratedClass* PythonGeneratedClass = CastChecked<UPythonGeneratedClass>(InObject);
-				bool bMarkClassPendingKill = true; // Mark types as PendingKill if they have no instances left (excluding their CDO)
+				bool bMarkClassPendingKill = false; // Mark types as PendingKill if they have no instances left (excluding their CDO)
 
 				ForEachObjectOfClass(PythonGeneratedClass, [&bMarkClassPendingKill, &FlagObjectForPurge](UObject* InInnerObject)
 				{
@@ -121,7 +121,7 @@ void FPyReferenceCollector::PurgeUnrealGeneratedTypes()
 		{
 			ForEachObjectOfClass(UPythonGeneratedStruct::StaticClass(), [&FlagObjectForPurge](UObject* InObject)
 			{
-				FlagObjectForPurge(InObject, /*bMarkPendingKill*/true);
+				FlagObjectForPurge(InObject, /*bMarkPendingKill*/false);
 			}, false);
 		}
 
@@ -130,7 +130,7 @@ void FPyReferenceCollector::PurgeUnrealGeneratedTypes()
 		{
 			ForEachObjectOfClass(UPythonGeneratedEnum::StaticClass(), [&FlagObjectForPurge](UObject* InObject)
 			{
-				FlagObjectForPurge(InObject, /*bMarkPendingKill*/true);
+				FlagObjectForPurge(InObject, /*bMarkPendingKill*/false);
 			}, false);
 		}
 
@@ -142,7 +142,7 @@ void FPyReferenceCollector::PurgeUnrealGeneratedTypes()
 
 			for (UClass* PythonCallableClass : PythonCallableClasses)
 			{
-				bool bMarkClassPendingKill = true; // Mark types as PendingKill if they have no instances left (excluding their CDO)
+				bool bMarkClassPendingKill = false; // Mark types as PendingKill if they have no instances left (excluding their CDO)
 
 				ForEachObjectOfClass(PythonCallableClass, [&bMarkClassPendingKill, &FlagObjectForPurge](UObject* InObject)
 				{
@@ -186,7 +186,7 @@ void FPyReferenceCollector::PurgeUnrealGeneratedTypes()
 	//  - Pass 2 will purge any types that no longer have instances left after running pass 1
 	if (RunPurgeUnrealGeneratedTypes(/*bLogFailure*/false))
 	{
-		RunPurgeUnrealGeneratedTypes(/*bLogFailure*/true);
+		RunPurgeUnrealGeneratedTypes(/*bLogFailure*/false);
 	}
 }
 
