Fixed delegate wrappers being incorrectly GC'd when no instances existed in Python

diff --git a/Engine/Plugins/Experimental/PythonScriptPlugin/Source/PythonScriptPlugin/Private/PyWrapperDelegate.cpp b/Engine/Plugins/Experimental/PythonScriptPlugin/Source/PythonScriptPlugin/Private/PyWrapperDelegate.cpp
--- a/Engine/Plugins/Experimental/PythonScriptPlugin/Source/PythonScriptPlugin/Private/PyWrapperDelegate.cpp
+++ b/Engine/Plugins/Experimental/PythonScriptPlugin/Source/PythonScriptPlugin/Private/PyWrapperDelegate.cpp
@@ -1289,8 +1289,6 @@ PyTypeObject PyWrapperMulticastDelegateType = InitializePyWrapperMulticastDelega
 
 void FPyWrapperDelegateMetaData::AddReferencedObjects(FPyWrapperBase* Instance, FReferenceCollector& Collector)
 {
-	TPyWrapperDelegateMetaData<FPyWrapperDelegate>::AddReferencedObjects(Instance, Collector);
-
 	FPyWrapperDelegate* Self = static_cast<FPyWrapperDelegate*>(Instance);
 	if (Self->DelegateInstance)
 	{
@@ -1300,8 +1298,6 @@ void FPyWrapperDelegateMetaData::AddReferencedObjects(FPyWrapperBase* Instance,
 
 void FPyWrapperMulticastDelegateMetaData::AddReferencedObjects(FPyWrapperBase* Instance, FReferenceCollector& Collector)
 {
-	TPyWrapperDelegateMetaData<FPyWrapperMulticastDelegate>::AddReferencedObjects(Instance, Collector);
-
 	FPyWrapperMulticastDelegate* Self = static_cast<FPyWrapperMulticastDelegate*>(Instance);
 	if (Self->DelegateInstance)
 	{
diff --git a/Engine/Plugins/Experimental/PythonScriptPlugin/Source/PythonScriptPlugin/Private/PyWrapperDelegate.h b/Engine/Plugins/Experimental/PythonScriptPlugin/Source/PythonScriptPlugin/Private/PyWrapperDelegate.h
--- a/Engine/Plugins/Experimental/PythonScriptPlugin/Source/PythonScriptPlugin/Private/PyWrapperDelegate.h
+++ b/Engine/Plugins/Experimental/PythonScriptPlugin/Source/PythonScriptPlugin/Private/PyWrapperDelegate.h
@@ -105,12 +105,6 @@ struct TPyWrapperDelegateMetaData : public FPyWrapperBaseMetaData
 		return GetPythonCallableForDelegateClass(Py_TYPE(Instance));
 	}
 
-	/** Add object references from the given Python object to the given collector */
-	virtual void AddReferencedObjects(FPyWrapperBase* Instance, FReferenceCollector& Collector) override
-	{
-		Collector.AddReferencedObject(PythonCallableForDelegateClass);
-	}
-
 	/** Get the reflection meta data type object associated with this wrapper type if there is one or nullptr if not. */
 	virtual const UField* GetMetaType() const override
 	{
diff --git a/Engine/Plugins/Experimental/PythonScriptPlugin/Source/PythonScriptPlugin/Private/PyWrapperTypeRegistry.cpp b/Engine/Plugins/Experimental/PythonScriptPlugin/Source/PythonScriptPlugin/Private/PyWrapperTypeRegistry.cpp
--- a/Engine/Plugins/Experimental/PythonScriptPlugin/Source/PythonScriptPlugin/Private/PyWrapperTypeRegistry.cpp
+++ b/Engine/Plugins/Experimental/PythonScriptPlugin/Source/PythonScriptPlugin/Private/PyWrapperTypeRegistry.cpp
@@ -1795,7 +1795,7 @@ PyTypeObject* FPyWrapperTypeRegistry::GenerateWrappedDelegateType(const UFunctio
 	// Generate the proxy class needed to wrap Python callables in Unreal delegates
 	UClass* PythonCallableForDelegateClass = nullptr;
 	{
-		PythonCallableForDelegateClass = NewObject<UClass>(GetPythonTypeContainer(), *FString::Printf(TEXT("%s__PythonCallable"), *DelegateBaseTypename), RF_Public);
+		PythonCallableForDelegateClass = NewObject<UClass>(GetPythonTypeContainer(), *FString::Printf(TEXT("%s__PythonCallable"), *DelegateBaseTypename), RF_Public | RF_Standalone | RF_Transient);
 		UFunction* PythonCallableForDelegateFunc = nullptr;
 		{
 			FObjectDuplicationParameters FuncDuplicationParams(const_cast<UFunction*>(InDelegateSignature), PythonCallableForDelegateClass);
