Calling a Python-generated function no longer crashes when called from Blueprints, or with non-POD output data

The old logic for dealing with the parameters within the Unreal function thunk (UPythonGeneratedClass::CallPythonFunction and UPythonCallableForDelegate::CallPythonNative) assumed that Stack.Locals would be set-up correctly for input, and could be written to for output. This isn't true, as Stack.Locals is populated by a memcpy in UObject::ProcessEvent (so isn't safe to write to for non-POD data), and Blueprints just pass along their existing frame data since they skip the call to ProcessScriptFunction (as Python generated functions are marked FUNC_Native).

This change adds a common thunk processing function (PyGenUtil::InvokePythonCallableFromUnrealFunctionThunk) which is used when calling into Python from an Unreal function thunk. This function steps the data off the stack frame for each property, copying input data to Python, and building the output chain (FOutParmRec) from output data addresses. This matches what the UHT generated thunk does for native functions (via P_GET_PROPERTY and P_GET_PROPERTY_REF), as well as what ProcessScriptFunction does for functions with Blueprint bytecode.

diff --git a/Engine/Plugins/Experimental/PythonScriptPlugin/Source/PythonScriptPlugin/Private/PyGenUtil.cpp b/Engine/Plugins/Experimental/PythonScriptPlugin/Source/PythonScriptPlugin/Private/PyGenUtil.cpp
--- a/Engine/Plugins/Experimental/PythonScriptPlugin/Source/PythonScriptPlugin/Private/PyGenUtil.cpp
+++ b/Engine/Plugins/Experimental/PythonScriptPlugin/Source/PythonScriptPlugin/Private/PyGenUtil.cpp
@@ -9,6 +9,7 @@
 #include "PyWrapperStruct.h"
 #include "Internationalization/BreakIterator.h"
 #include "Misc/Paths.h"
+#include "Misc/ScopeExit.h"
 #include "Misc/FileHelper.h"
 #include "UObject/Class.h"
 #include "UObject/Package.h"
@@ -953,32 +954,32 @@ PyObject* PackReturnValues(const void* InBaseParamsAddr, const TArray<FGenerated
 	}
 }
 
-bool UnpackReturnValues(PyObject* InRetVals, void* InBaseParamsAddr, const TArray<FGeneratedWrappedMethodParameter>& InOutputParams, const TCHAR* InErrorCtxt, const TCHAR* InCallingCtxt)
+bool UnpackReturnValues(PyObject* InRetVals, const FOutParmRec* InOutputParms, const TCHAR* InErrorCtxt, const TCHAR* InCallingCtxt)
 {
-	if (!InOutputParams.Num())
+	if (!InOutputParms)
 	{
 		return true;
 	}
 
-	int32 ReturnPropIndex = 0;
+	const FOutParmRec* OutParamRec = InOutputParms;
 
 	// If we have multiple return values and the main return value is a bool, we expect None (for false) or the (potentially packed) return value without the bool (for true)
-	if (InOutputParams.Num() > 1 && InOutputParams[0].ParamProp->HasAnyPropertyFlags(CPF_ReturnParm) && InOutputParams[0].ParamProp->IsA<UBoolProperty>())
+	if (OutParamRec->NextOutParm && OutParamRec->Property->HasAnyPropertyFlags(CPF_ReturnParm) && OutParamRec->Property->IsA<UBoolProperty>())
 	{
-		const UBoolProperty* BoolReturn = CastChecked<const UBoolProperty>(InOutputParams[0].ParamProp);
+		const UBoolProperty* BoolReturn = CastChecked<const UBoolProperty>(OutParamRec->Property);
 		const bool bReturnValue = InRetVals != Py_None;
-		BoolReturn->SetPropertyValue(BoolReturn->ContainerPtrToValuePtr<void>(InBaseParamsAddr), bReturnValue);
+		BoolReturn->SetPropertyValue(OutParamRec->PropAddr, bReturnValue);
 
-		ReturnPropIndex = 1; // Start unpacking at the 1st out value
+		OutParamRec = OutParamRec->NextOutParm; // Start unpacking at the 1st out value
+		check(OutParamRec);
 	}
 
 	// Do we need to expect a packed tuple, or just a single value?
-	const int32 NumPropertiesToUnpack = InOutputParams.Num() - ReturnPropIndex;
-	if (NumPropertiesToUnpack == 1)
+	if (OutParamRec->NextOutParm == nullptr)
 	{
-		if (!PyConversion::NativizeProperty_InContainer(InRetVals, InOutputParams[ReturnPropIndex].ParamProp, InBaseParamsAddr, 0))
+		if (!PyConversion::NativizeProperty_Direct(InRetVals, OutParamRec->Property, OutParamRec->PropAddr))
 		{
-			PyUtil::SetPythonError(PyExc_TypeError, InErrorCtxt, *FString::Printf(TEXT("Failed to convert return property '%s' (%s) when calling %s"), *InOutputParams[ReturnPropIndex].ParamProp->GetName(), *InOutputParams[ReturnPropIndex].ParamProp->GetClass()->GetName(), InCallingCtxt));
+			PyUtil::SetPythonError(PyExc_TypeError, InErrorCtxt, *FString::Printf(TEXT("Failed to convert return property '%s' (%s) when calling %s"), *OutParamRec->Property->GetName(), *OutParamRec->Property->GetClass()->GetName(), InCallingCtxt));
 			return false;
 		}
 	}
@@ -990,6 +991,13 @@ bool UnpackReturnValues(PyObject* InRetVals, void* InBaseParamsAddr, const TArra
 			return false;
 		}
 
+		check(OutParamRec->NextOutParm);
+		int32 NumPropertiesToUnpack = 2; // We can start from 2 since we know we already have at least 2 output parameters to get to this code
+		for (const FOutParmRec* Tmp = OutParamRec->NextOutParm->NextOutParm; Tmp; Tmp = Tmp->NextOutParm)
+		{
+			++NumPropertiesToUnpack;
+		}
+
 		const int32 RetTupleSize = PyTuple_Size(InRetVals);
 		if (RetTupleSize != NumPropertiesToUnpack)
 		{
@@ -998,15 +1006,152 @@ bool UnpackReturnValues(PyObject* InRetVals, void* InBaseParamsAddr, const TArra
 		}
 
 		int32 RetTupleIndex = 0;
-		for (; ReturnPropIndex < InOutputParams.Num(); ++ReturnPropIndex)
+		do
 		{
 			PyObject* RetVal = PyTuple_GetItem(InRetVals, RetTupleIndex++);
-			if (!PyConversion::NativizeProperty_InContainer(RetVal, InOutputParams[ReturnPropIndex].ParamProp, InBaseParamsAddr, 0))
+			if (!PyConversion::NativizeProperty_Direct(RetVal, OutParamRec->Property, OutParamRec->PropAddr))
 			{
-				PyUtil::SetPythonError(PyExc_TypeError, InErrorCtxt, *FString::Printf(TEXT("Failed to convert return property '%s' (%s) when calling %s"), *InOutputParams[ReturnPropIndex].ParamProp->GetName(), *InOutputParams[ReturnPropIndex].ParamProp->GetClass()->GetName(), InCallingCtxt));
+				PyUtil::SetPythonError(PyExc_TypeError, InErrorCtxt, *FString::Printf(TEXT("Failed to convert return property '%s' (%s) when calling %s"), *OutParamRec->Property->GetName(), *OutParamRec->Property->GetClass()->GetName(), InCallingCtxt));
 				return false;
 			}
+			OutParamRec = OutParamRec->NextOutParm;
 		}
+		while (OutParamRec);
+	}
+
+	return true;
+}
+
+bool InvokePythonCallableFromUnrealFunctionThunk(FPyObjectPtr InSelf, PyObject* InCallable, const UFunction* InFunc, UObject* Context, FFrame& Stack, RESULT_DECL)
+{
+	// Allocate memory to store our local argument data
+	void* LocalStruct = FMemory_Alloca(InFunc->GetStructureSize());
+	InFunc->InitializeStruct(LocalStruct);
+	ON_SCOPE_EXIT
+	{
+		InFunc->DestroyStruct(LocalStruct);
+	};
+
+	// Stores information about inputs and outputs
+	FOutParmRec* OutParms = nullptr;
+	TArray<FPyObjectPtr, TInlineAllocator<4>> PyParams;
+
+	// Add any return property to the output params chain
+	if (UProperty* ReturnProp = InFunc->GetReturnProperty())
+	{
+		FOutParmRec* Out = (FOutParmRec*)FMemory_Alloca(sizeof(FOutParmRec));
+		Out->Property = ReturnProp;
+		Out->PropAddr = (uint8*)RESULT_PARAM;
+
+		// Link it to the head of the list, as UnpackReturnValues expects the return value to be first in the list
+		Out->NextOutParm = OutParms;
+		OutParms = Out;
+	}
+
+	// Get the value of the input params for the Python args, and cache the addresses that return and output data should be unpacked to
+	bool bProcessedInputs = true;
+	{
+		int32 ArgIndex = 0;
+		FOutParmRec** LastOut = &OutParms;
+
+		// We iterate the fields directly here as we need to process input and output properties in the 
+		// correct stack order, as we're potentially popping data off the bytecode stack
+		for (TFieldIterator<UProperty> ParamIt(InFunc); ParamIt; ++ParamIt)
+		{
+			UProperty* Param = *ParamIt;
+
+			// Skip the return value; it never has data on the bytecode stack and was added to the output params chain before this loop
+			if (Param->HasAnyPropertyFlags(CPF_ReturnParm))
+			{
+				continue;
+			}
+
+			// Step the property data to populate the local value
+			Stack.MostRecentPropertyAddress = nullptr;
+			void* LocalValue = Param->ContainerPtrToValuePtr<void>(LocalStruct);
+			Stack.StepCompiledIn<UProperty>(LocalValue);
+
+			void* ValueReadAddress = LocalValue;
+
+			// Add any output parameters to the output params chain
+			if (PyUtil::IsOutputParameter(Param))
+			{
+				CA_SUPPRESS(6263) // using _alloca in a loop
+				FOutParmRec* Out = (FOutParmRec*)FMemory_Alloca(sizeof(FOutParmRec));
+				Out->Property = Param;
+				Out->PropAddr = (Stack.MostRecentPropertyAddress) ? Stack.MostRecentPropertyAddress : (uint8*)LocalValue;
+				Out->NextOutParm = nullptr;
+
+				// If this value is also an input param (eg, UPARAM(ref)) then we need to read its input value from the output address too
+				ValueReadAddress = Out->PropAddr;
+
+				// Link it to the end of the list
+				if (*LastOut)
+				{
+					(*LastOut)->NextOutParm = Out;
+					LastOut = &(*LastOut)->NextOutParm;
+				}
+				else
+				{
+					*LastOut = Out;
+				}
+			}
+
+			// Convert any input parameters for use with Python
+			if (PyUtil::IsInputParameter(Param))
+			{
+				FPyObjectPtr& PyParam = PyParams.AddDefaulted_GetRef();
+				if (!PyConversion::PythonizeProperty_Direct(Param, ValueReadAddress, PyParam.Get()))
+				{
+					PyUtil::SetPythonError(PyExc_TypeError, InCallable ? *PyUtil::GetErrorContext(InCallable) : TEXT("<null>"), *FString::Printf(TEXT("Failed to convert argument at pos '%d' when calling function '%s' on '%s'"), ArgIndex + 1, *InFunc->GetName(), *P_THIS_OBJECT->GetName()));
+					bProcessedInputs = false;
+				}
+				++ArgIndex;
+			}
+		}
+	}
+
+	// Validate we reached the end of the parameters when stepping the bytecode stack
+	if (Stack.Code)
+	{
+		checkSlow(*Stack.Code == EX_EndFunctionParms);
+		++Stack.Code;
+	}
+
+	// If any errors happened during parameter processing then we can bail now that we've finished stepping the bytecode stack
+	// We can also bail if we have no Python callable to invoke
+	if (!bProcessedInputs || !InCallable)
+	{
+		return false;
+	}
+
+	// Prepare the arguments tuple for the Python callable
+	FPyObjectPtr PyArgs;
+	if (InSelf || PyParams.Num() > 0)
+	{
+		const int32 PyParamOffset = (InSelf ? 1 : 0);
+		PyArgs = FPyObjectPtr::StealReference(PyTuple_New(PyParams.Num() + PyParamOffset));
+		if (InSelf)
+		{
+			PyTuple_SetItem(PyArgs, 0, InSelf.Release()); // SetItem steals the reference
+		}
+		for (int32 PyParamIndex = 0; PyParamIndex < PyParams.Num(); ++PyParamIndex)
+		{
+			PyTuple_SetItem(PyArgs, PyParamIndex + PyParamOffset, PyParams[PyParamIndex].Release()); // SetItem steals the reference
+		}
+	}
+
+	// Invoke the Python callable
+	FPyObjectPtr RetVals = FPyObjectPtr::StealReference(PyObject_CallObject(InCallable, PyArgs));
+	if (!RetVals)
+	{
+		return false;
+	}
+
+	// Unpack any output values
+	if (!UnpackReturnValues(RetVals, OutParms, *PyUtil::GetErrorContext(InCallable), *FString::Printf(TEXT("function '%s' on '%s'"), *InFunc->GetName(), *P_THIS_OBJECT->GetName())))
+	{
+		return false;
 	}
 
 	return true;
diff --git a/Engine/Plugins/Experimental/PythonScriptPlugin/Source/PythonScriptPlugin/Private/PyGenUtil.h b/Engine/Plugins/Experimental/PythonScriptPlugin/Source/PythonScriptPlugin/Private/PyGenUtil.h
--- a/Engine/Plugins/Experimental/PythonScriptPlugin/Source/PythonScriptPlugin/Private/PyGenUtil.h
+++ b/Engine/Plugins/Experimental/PythonScriptPlugin/Source/PythonScriptPlugin/Private/PyGenUtil.h
@@ -756,11 +756,7 @@ namespace PyGenUtil
 	/** Stores the data needed by a runtime generated Python enum entry */
 	struct FGeneratedWrappedEnumEntry
 	{
-		FGeneratedWrappedEnumEntry()
-			: EntryValue(0)
-		{
-		}
-
+		FGeneratedWrappedEnumEntry() = default;
 		FGeneratedWrappedEnumEntry(FGeneratedWrappedEnumEntry&&) = default;
 		FGeneratedWrappedEnumEntry(const FGeneratedWrappedEnumEntry&) = default;
 		FGeneratedWrappedEnumEntry& operator=(FGeneratedWrappedEnumEntry&&) = default;
@@ -773,7 +769,7 @@ namespace PyGenUtil
 		FUTF8Buffer EntryDoc;
 
 		/** The value of the entry */
-		int64 EntryValue;
+		int64 EntryValue = 0;
 	};
 
 	/** Stores the data needed by a runtime generated Python enum type */
@@ -830,7 +826,7 @@ namespace PyGenUtil
 		FPyObjectPtr PyFunction;
 
 		/** Is this a function hidden from Python? (eg, internal getter/setter function) */
-		bool bIsHidden;
+		bool bIsHidden = false;
 	};
 
 	/** How should PythonizeName adjust the final name? */
@@ -907,7 +903,10 @@ namespace PyGenUtil
 	PyObject* PackReturnValues(const void* InBaseParamsAddr, const TArray<FGeneratedWrappedMethodParameter>& InOutputParams, const TCHAR* InErrorCtxt, const TCHAR* InCallingCtxt);
 
 	/** Given a Python return value, unpack the values into the struct data associated with them */
-	bool UnpackReturnValues(PyObject* InRetVals, void* InBaseParamsAddr, const TArray<FGeneratedWrappedMethodParameter>& InOutputParams, const TCHAR* InErrorCtxt, const TCHAR* InCallingCtxt);
+	bool UnpackReturnValues(PyObject* InRetVals, const FOutParmRec* InOutputParms, const TCHAR* InErrorCtxt, const TCHAR* InCallingCtxt);
+
+	/** Invoke a Python callable from an Unreal function thunk (ie, DEFINE_FUNCTION), ensuring that the calling convention is correct for both native and script function callers */
+	bool InvokePythonCallableFromUnrealFunctionThunk(FPyObjectPtr InSelf, PyObject* InCallable, const UFunction* InFunc, UObject* Context, FFrame& Stack, RESULT_DECL);
 
 	/** Get the current value of the given property from the given struct */
 	PyObject* GetPropertyValue(const UStruct* InStruct, void* InStructData, const FGeneratedWrappedProperty& InPropDef, const char *InAttributeName, PyObject* InOwnerPyObject, const TCHAR* InErrorCtxt);
diff --git a/Engine/Plugins/Experimental/PythonScriptPlugin/Source/PythonScriptPlugin/Private/PyWrapperDelegate.cpp b/Engine/Plugins/Experimental/PythonScriptPlugin/Source/PythonScriptPlugin/Private/PyWrapperDelegate.cpp
--- a/Engine/Plugins/Experimental/PythonScriptPlugin/Source/PythonScriptPlugin/Private/PyWrapperDelegate.cpp
+++ b/Engine/Plugins/Experimental/PythonScriptPlugin/Source/PythonScriptPlugin/Private/PyWrapperDelegate.cpp
@@ -18,81 +18,16 @@ const FName UPythonCallableForDelegate::GeneratedFuncName = "CallPython";
 DEFINE_FUNCTION(UPythonCallableForDelegate::CallPythonNative)
 {
 #if WITH_PYTHON
-	if (P_THIS->PyCallable)
-	{
-		auto DoCall = [&]() -> bool
-		{
-			if (Stack.Node->Children == nullptr)
-			{
-				// Simple case, no parameters or return value
-				FPyObjectPtr RetVals = FPyObjectPtr::StealReference(PyObject_CallObject(P_THIS->PyCallable, nullptr));
-				if (!RetVals)
-				{
-					return false;
-				}
-			}
-			else
-			{
-				PyGenUtil::FGeneratedWrappedFunction DelegateFuncDef;
-				DelegateFuncDef.SetFunction(Stack.Node, PyGenUtil::FGeneratedWrappedFunction::SFF_ExtractParameters);
+	// Note: This function *must not* return until InvokePythonCallableFromUnrealFunctionThunk has been called, as we need to step over the correct amount of data from the bytecode stack!
 
-				// Complex case, parameters or return value
-				TArray<FPyObjectPtr, TInlineAllocator<4>> PyParams;
+	const UFunction* Func = Stack.CurrentNativeFunction;
 
-				// Get the value of the input params for the Python args
-				{
-					int32 ArgIndex = 0;
-					for (const PyGenUtil::FGeneratedWrappedMethodParameter& ParamDef : DelegateFuncDef.InputParams)
-					{
-						FPyObjectPtr& PyParam = PyParams.AddDefaulted_GetRef();
-						if (!PyConversion::PythonizeProperty_InContainer(ParamDef.ParamProp, Stack.Locals, 0, PyParam.Get()))
-						{
-							PyUtil::SetPythonError(PyExc_TypeError, P_THIS->PyCallable, *FString::Printf(TEXT("Failed to convert argument at pos '%d' when calling function '%s' on '%s'"), ArgIndex + 1, *Stack.Node->GetName(), *P_THIS_OBJECT->GetName()));
-							return false;
-						}
-						++ArgIndex;
-					}
-				}
-
-				FPyObjectPtr PyArgs = FPyObjectPtr::StealReference(PyTuple_New(PyParams.Num()));
-				for (int32 PyParamIndex = 0; PyParamIndex < PyParams.Num(); ++PyParamIndex)
-				{
-					PyTuple_SetItem(PyArgs, PyParamIndex, PyParams[PyParamIndex].Release()); // SetItem steals the reference
-				}
-
-				FPyObjectPtr RetVals = FPyObjectPtr::StealReference(PyObject_CallObject(P_THIS->PyCallable, PyArgs));
-				if (!RetVals)
-				{
-					return false;
-				}
-
-				if (!PyGenUtil::UnpackReturnValues(RetVals, Stack.Locals, DelegateFuncDef.OutputParams, *PyUtil::GetErrorContext(P_THIS->PyCallable), *FString::Printf(TEXT("function '%s' on '%s'"), *Stack.Node->GetName(), *P_THIS_OBJECT->GetName())))
-				{
-					return false;
-				}
-
-				// Copy the data back out of the function call
-				if (const UProperty* ReturnProp = Stack.Node->GetReturnProperty())
-				{
-					ReturnProp->CopyCompleteValue(RESULT_PARAM, ReturnProp->ContainerPtrToValuePtr<void>(Stack.Locals));
-				}
-				for (FOutParmRec* OutParamRec = Stack.OutParms; OutParamRec; OutParamRec = OutParamRec->NextOutParm)
-				{
-					OutParamRec->Property->CopyCompleteValue(OutParamRec->PropAddr, OutParamRec->Property->ContainerPtrToValuePtr<void>(Stack.Locals));
-				}
-			}
-
-			return true;
-		};
-
-		// Execute Python code within this block
+	// Execute Python code within this block
+	{
+		FPyScopedGIL GIL;
+		if (!PyGenUtil::InvokePythonCallableFromUnrealFunctionThunk(FPyObjectPtr(), P_THIS->PyCallable, Func, Context, Stack, RESULT_PARAM))
 		{
-			FPyScopedGIL GIL;
-
-			if (!DoCall())
-			{
-				PyUtil::ReThrowPythonError();
-			}
+			PyUtil::ReThrowPythonError();
 		}
 	}
 #endif	// WITH_PYTHON
diff --git a/Engine/Plugins/Experimental/PythonScriptPlugin/Source/PythonScriptPlugin/Private/PyWrapperObject.cpp b/Engine/Plugins/Experimental/PythonScriptPlugin/Source/PythonScriptPlugin/Private/PyWrapperObject.cpp
--- a/Engine/Plugins/Experimental/PythonScriptPlugin/Source/PythonScriptPlugin/Private/PyWrapperObject.cpp
+++ b/Engine/Plugins/Experimental/PythonScriptPlugin/Source/PythonScriptPlugin/Private/PyWrapperObject.cpp
@@ -1834,114 +1834,46 @@ UPythonGeneratedClass* UPythonGeneratedClass::ReparentClass(UPythonGeneratedClas
 
 DEFINE_FUNCTION(UPythonGeneratedClass::CallPythonFunction)
 {
-	// Get the correct class from the UFunction so that we can perform static dispatch to the correct type
-	const UPythonGeneratedClass* This = CastChecked<UPythonGeneratedClass>(Stack.Node->GetOwnerClass());
+	// Note: This function *must not* return until InvokePythonCallableFromUnrealFunctionThunk has been called, as we need to step over the correct amount of data from the bytecode stack!
+
+	const UFunction* Func = Stack.CurrentNativeFunction;
 
 	// Find the Python function to call
 	TSharedPtr<PyGenUtil::FFunctionDef> FuncDef;
 	{
-		const TSharedPtr<PyGenUtil::FFunctionDef>* FuncDefPtr = This->FunctionDefs.FindByPredicate([&Stack](const TSharedPtr<PyGenUtil::FFunctionDef>& InFuncDef)
+		// Get the correct class from the UFunction so that we can perform static dispatch to the correct type
+		const UPythonGeneratedClass* This = CastChecked<UPythonGeneratedClass>(Func->GetOwnerClass());
+
+		const TSharedPtr<PyGenUtil::FFunctionDef>* FuncDefPtr = This->FunctionDefs.FindByPredicate([Func](const TSharedPtr<PyGenUtil::FFunctionDef>& InFuncDef)
 		{
-			return InFuncDef->GeneratedWrappedMethod.MethodFunc.Func == Stack.Node;
+			return InFuncDef->GeneratedWrappedMethod.MethodFunc.Func == Func;
 		});
 		FuncDef = FuncDefPtr ? *FuncDefPtr : nullptr;
-	}
-	if (!FuncDef.IsValid())
-	{
-		UE_LOG(LogPython, Error, TEXT("Failed to find Python function for '%s' on '%s'"), *Stack.Node->GetName(), *This->GetName());
+
+		if (!FuncDef.IsValid())
+		{
+			UE_LOG(LogPython, Error, TEXT("Failed to find Python function for '%s' on '%s'"), *Func->GetName(), *This->GetName());
+		}
 	}
 
 	// Find the Python object to call the function on
 	FPyObjectPtr PySelf;
-	if (!Stack.Node->HasAnyFunctionFlags(FUNC_Static))
+	bool bSelfError = false;
+	if (!Func->HasAnyFunctionFlags(FUNC_Static))
 	{
 		FPyScopedGIL GIL;
 		PySelf = FPyObjectPtr::StealReference((PyObject*)FPyWrapperObjectFactory::Get().CreateInstance(P_THIS_OBJECT));
 		if (!PySelf)
 		{
 			UE_LOG(LogPython, Error, TEXT("Failed to create a Python wrapper for '%s'"), *P_THIS_OBJECT->GetName());
-			return;
+			bSelfError = true;
 		}
 	}
 
-	auto DoCall = [&]() -> bool
-	{
-		if (Stack.Node->Children == nullptr)
-		{
-			// Simple case, no parameters or return value
-			FPyObjectPtr PyArgs;
-			if (PySelf)
-			{
-				PyArgs = FPyObjectPtr::StealReference(PyTuple_New(1));
-				PyTuple_SetItem(PyArgs, 0, PySelf.Release()); // SetItem steals the reference
-			}
-			FPyObjectPtr RetVals = FPyObjectPtr::StealReference(PyObject_CallObject(FuncDef->PyFunction, PyArgs));
-			if (!RetVals)
-			{
-				return false;
-			}
-		}
-		else
-		{
-			// Complex case, parameters or return value
-			TArray<FPyObjectPtr, TInlineAllocator<4>> PyParams;
-
-			// Get the value of the input params for the Python args
-			{
-				int32 ArgIndex = 0;
-				for (const PyGenUtil::FGeneratedWrappedMethodParameter& ParamDef : FuncDef->GeneratedWrappedMethod.MethodFunc.InputParams)
-				{
-					FPyObjectPtr& PyParam = PyParams.AddDefaulted_GetRef();
-					if (!PyConversion::PythonizeProperty_InContainer(ParamDef.ParamProp, Stack.Locals, 0, PyParam.Get()))
-					{
-						PyUtil::SetPythonError(PyExc_TypeError, FuncDef->PyFunction, *FString::Printf(TEXT("Failed to convert argument at pos '%d' when calling function '%s' on '%s'"), ArgIndex + 1, *Stack.Node->GetName(), *P_THIS_OBJECT->GetName()));
-						return false;
-					}
-					++ArgIndex;
-				}
-			}
-
-			const int32 PyParamOffset = (PySelf ? 1 : 0);
-			FPyObjectPtr PyArgs = FPyObjectPtr::StealReference(PyTuple_New(PyParams.Num() + PyParamOffset));
-			if (PySelf)
-			{
-				PyTuple_SetItem(PyArgs, 0, PySelf.Release()); // SetItem steals the reference
-			}
-			for (int32 PyParamIndex = 0; PyParamIndex < PyParams.Num(); ++PyParamIndex)
-			{
-				PyTuple_SetItem(PyArgs, PyParamIndex + PyParamOffset, PyParams[PyParamIndex].Release()); // SetItem steals the reference
-			}
-
-			FPyObjectPtr RetVals = FPyObjectPtr::StealReference(PyObject_CallObject(FuncDef->PyFunction, PyArgs));
-			if (!RetVals)
-			{
-				return false;
-			}
-
-			if (!PyGenUtil::UnpackReturnValues(RetVals, Stack.Locals, FuncDef->GeneratedWrappedMethod.MethodFunc.OutputParams, *PyUtil::GetErrorContext(FuncDef->PyFunction), *FString::Printf(TEXT("function '%s' on '%s'"), *Stack.Node->GetName(), *P_THIS_OBJECT->GetName())))
-			{
-				return false;
-			}
-
-			// Copy the data back out of the function call
-			if (const UProperty* ReturnProp = Stack.Node->GetReturnProperty())
-			{
-				ReturnProp->CopyCompleteValue(RESULT_PARAM, ReturnProp->ContainerPtrToValuePtr<void>(Stack.Locals));
-			}
-			for (FOutParmRec* OutParamRec = Stack.OutParms; OutParamRec; OutParamRec = OutParamRec->NextOutParm)
-			{
-				OutParamRec->Property->CopyCompleteValue(OutParamRec->PropAddr, OutParamRec->Property->ContainerPtrToValuePtr<void>(Stack.Locals));
-			}
-		}
-
-		return true;
-	};
-
 	// Execute Python code within this block
 	{
 		FPyScopedGIL GIL;
-
-		if (!DoCall())
+		if (!PyGenUtil::InvokePythonCallableFromUnrealFunctionThunk(PySelf, FuncDef ? FuncDef->PyFunction.GetPtr() : nullptr, Func, Context, Stack, RESULT_PARAM) || bSelfError)
 		{
 			PyUtil::ReThrowPythonError();
 		}
