Added additional tests for Python-generated overrides and non-POD data (re: CL# 7064650)

#rnx

diff --git a/Engine/Plugins/Experimental/PythonScriptPlugin/Source/PythonScriptPlugin/Private/PyTest.cpp b/Engine/Plugins/Experimental/PythonScriptPlugin/Source/PythonScriptPlugin/Private/PyTest.cpp
--- a/Engine/Plugins/Experimental/PythonScriptPlugin/Source/PythonScriptPlugin/Private/PyTest.cpp
+++ b/Engine/Plugins/Experimental/PythonScriptPlugin/Source/PythonScriptPlugin/Private/PyTest.cpp
@@ -61,6 +61,10 @@ int32 UPyTestObject::FuncBlueprintNative_Implementation(const int32 InValue) con
 	return InValue;
 }
 
+void UPyTestObject::FuncBlueprintNativeRef_Implementation(FPyTestStruct& InOutStruct) const
+{
+}
+
 int32 UPyTestObject::CallFuncBlueprintImplementable(const int32 InValue) const
 {
 	return FuncBlueprintImplementable(InValue);
@@ -71,6 +75,11 @@ int32 UPyTestObject::CallFuncBlueprintNative(const int32 InValue) const
 	return FuncBlueprintNative(InValue);
 }
 
+void UPyTestObject::CallFuncBlueprintNativeRef(FPyTestStruct& InOutStruct) const
+{
+	return FuncBlueprintNativeRef(InOutStruct);
+}
+
 void UPyTestObject::FuncTakingPyTestStruct(const FPyTestStruct& InStruct) const
 {
 }
diff --git a/Engine/Plugins/Experimental/PythonScriptPlugin/Source/PythonScriptPlugin/Private/PyTest.h b/Engine/Plugins/Experimental/PythonScriptPlugin/Source/PythonScriptPlugin/Private/PyTest.h
--- a/Engine/Plugins/Experimental/PythonScriptPlugin/Source/PythonScriptPlugin/Private/PyTest.h
+++ b/Engine/Plugins/Experimental/PythonScriptPlugin/Source/PythonScriptPlugin/Private/PyTest.h
@@ -181,12 +181,18 @@ public:
 	UFUNCTION(BlueprintNativeEvent, Category = "Python|Internal")
 	int32 FuncBlueprintNative(const int32 InValue) const;
 
+	UFUNCTION(BlueprintNativeEvent, Category = "Python|Internal")
+	void FuncBlueprintNativeRef(UPARAM(ref) FPyTestStruct& InOutStruct) const;
+
 	UFUNCTION(BlueprintCallable, Category = "Python|Internal")
 	int32 CallFuncBlueprintImplementable(const int32 InValue) const;
 
 	UFUNCTION(BlueprintCallable, Category = "Python|Internal")
 	int32 CallFuncBlueprintNative(const int32 InValue) const;
 
+	UFUNCTION(BlueprintCallable, Category = "Python|Internal")
+	void CallFuncBlueprintNativeRef(UPARAM(ref) FPyTestStruct& InOutStruct) const;
+
 	UFUNCTION(BlueprintCallable, Category = "Python|Internal")
 	void FuncTakingPyTestStruct(const FPyTestStruct& InStruct) const;
 
