CIS fix

#rnx

diff --git a/Engine/Plugins/Experimental/PythonScriptPlugin/Source/PythonScriptPlugin/Private/PyGenUtil.cpp b/Engine/Plugins/Experimental/PythonScriptPlugin/Source/PythonScriptPlugin/Private/PyGenUtil.cpp
--- a/Engine/Plugins/Experimental/PythonScriptPlugin/Source/PythonScriptPlugin/Private/PyGenUtil.cpp
+++ b/Engine/Plugins/Experimental/PythonScriptPlugin/Source/PythonScriptPlugin/Private/PyGenUtil.cpp
@@ -12,6 +12,7 @@
 #include "Misc/ScopeExit.h"
 #include "Misc/FileHelper.h"
 #include "UObject/Class.h"
+#include "UObject/Stack.h"
 #include "UObject/Package.h"
 #include "UObject/EnumProperty.h"
 #include "UObject/TextProperty.h"
diff --git a/Engine/Plugins/Experimental/PythonScriptPlugin/Source/PythonScriptPlugin/Private/PyGenUtil.h b/Engine/Plugins/Experimental/PythonScriptPlugin/Source/PythonScriptPlugin/Private/PyGenUtil.h
--- a/Engine/Plugins/Experimental/PythonScriptPlugin/Source/PythonScriptPlugin/Private/PyGenUtil.h
+++ b/Engine/Plugins/Experimental/PythonScriptPlugin/Source/PythonScriptPlugin/Private/PyGenUtil.h
@@ -13,6 +13,7 @@
 
 #if WITH_PYTHON
 
+struct FOutParmRec;
 struct FPyWrapperBaseMetaData;
 
 class FPyWrapperOwnerContext;
