Get a DLLHandle for Linux as well when loadind in the python module

diff --git a/Engine/Plugins/Experimental/PythonScriptPlugin/Source/PythonScriptPluginPreload/Private/PythonScriptPluginPreload.cpp b/Engine/Plugins/Experimental/PythonScriptPlugin/Source/PythonScriptPluginPreload/Private/PythonScriptPluginPreload.cpp
--- a/Engine/Plugins/Experimental/PythonScriptPlugin/Source/PythonScriptPluginPreload/Private/PythonScriptPluginPreload.cpp
+++ b/Engine/Plugins/Experimental/PythonScriptPlugin/Source/PythonScriptPluginPreload/Private/PythonScriptPluginPreload.cpp
@@ -35,10 +35,45 @@ public:
 
 private:
 #if WITH_PYTHON
-	void LoadPythonLibraries()
+	void LoadSharedDSO(const FString& PythonDSOWildcard, const FString& PythonDir)
 	{
+		auto FindPythonDSOs = [&PythonDSOWildcard](const FString& InPath)
+		{
+			TArray<FString> PythonDSONames;
+			IFileManager::Get().FindFiles(PythonDSONames, *(InPath / PythonDSOWildcard), true, false);
+			for (FString& PythonDSOName : PythonDSONames)
+			{
+				PythonDSOName = InPath / PythonDSOName;
+				FPaths::NormalizeFilename(PythonDSOName);
+			}
+			return PythonDSONames;
+		};
+
+		TArray<FString> PythonDSOPaths = FindPythonDSOs(PythonDir);
 #if PLATFORM_WINDOWS
-		// Load the DLLs
+		if (PythonDSOPaths.Num() == 0)
+		{
+			// If we didn't find anything, check the Windows directory as the DLLs can sometimes be installed there
+			FString WinDir = FPlatformMisc::GetEnvironmentVariable(TEXT("WINDIR"));
+			if (!WinDir.IsEmpty())
+			{
+				PythonDSOPaths = FindPythonDSOs(WinDir / TEXT("System32"));
+			}
+		}
+#endif
+
+		for (const FString& PythonDSOPath : PythonDSOPaths)
+		{
+			void* DLLHandle = FPlatformProcess::GetDllHandle(*PythonDSOPath);
+			check(DLLHandle != nullptr);
+			DLLHandles.Add(DLLHandle);
+		}
+	}
+
+	void LoadPythonLibraries()
+	{
+#if PLATFORM_WINDOWS || PLATFORM_LINUX
+		// Load the DSOs
 		{
 			// Build the full Python directory (UE_PYTHON_DIR may be relative to UE4 engine directory for portability)
 			FString PythonDir = UTF8_TO_TCHAR(UE_PYTHON_DIR);
@@ -46,38 +81,15 @@ private:
 			FPaths::NormalizeDirectoryName(PythonDir);
 			FPaths::RemoveDuplicateSlashes(PythonDir);
 
-			const FString PythonDllWildcard = FString::Printf(TEXT("python%d*.dll"), PY_MAJOR_VERSION);
-			auto FindPythonDLLs = [&PythonDllWildcard](const FString& InPath)
-			{
-				TArray<FString> PythonDLLNames;
-				IFileManager::Get().FindFiles(PythonDLLNames, *(InPath / PythonDllWildcard), true, false);
-				for (FString& PythonDLLName : PythonDLLNames)
-				{
-					PythonDLLName = InPath / PythonDLLName;
-					FPaths::NormalizeFilename(PythonDLLName);
-				}
-				return PythonDLLNames;
-			};
-
-			TArray<FString> PythonDLLPaths = FindPythonDLLs(PythonDir);
-			if (PythonDLLPaths.Num() == 0)
-			{
-				// If we didn't find anything, check the Windows directory as the DLLs can sometimes be installed there
-				FString WinDir = FPlatformMisc::GetEnvironmentVariable(TEXT("WINDIR"));
-				if (!WinDir.IsEmpty())
-				{
-					PythonDLLPaths = FindPythonDLLs(WinDir / TEXT("System32"));
-				}
-			}
-
-			for (const FString& PythonDLLPath : PythonDLLPaths)
-			{
-				void* DLLHandle = FPlatformProcess::GetDllHandle(*PythonDLLPath);
-				check(DLLHandle != nullptr);
-				DLLHandles.Add(DLLHandle);
-			}
+#if PLATFORM_WINDOWS
+			const FString PythonDSOWildcard = FString::Printf(TEXT("python%d*.dll"), PY_MAJOR_VERSION);
+#elif PLATFORM_LINUX
+			const FString PythonDSOWildcard = FString::Printf(TEXT("libpython%d*.so*"), PY_MAJOR_VERSION);
+			PythonDir /= "lib";
+#endif
+			LoadSharedDSO(PythonDSOWildcard, PythonDir);
 		}
-#endif	// PLATFORM_WINDOWS
+#endif	// PLATFORM_WINDOWS || PLATFORM_LINUX
 	}
 
 	void UnloadPythonLibraries()
