Fixed type validation of objects within containers

diff --git a/Engine/Plugins/Experimental/PythonScriptPlugin/Source/PythonScriptPlugin/Private/PyUtil.cpp b/Engine/Plugins/Experimental/PythonScriptPlugin/Source/PythonScriptPlugin/Private/PyUtil.cpp
--- a/Engine/Plugins/Experimental/PythonScriptPlugin/Source/PythonScriptPlugin/Private/PyUtil.cpp
+++ b/Engine/Plugins/Experimental/PythonScriptPlugin/Source/PythonScriptPlugin/Private/PyUtil.cpp
@@ -169,9 +169,19 @@ FPropertyDef::FPropertyDef(const UProperty* InProperty)
 	, KeyDef()
 	, ValueDef()
 {
+	if (const UObjectPropertyBase* ObjectProp = Cast<UObjectPropertyBase>(InProperty))
+	{
+		PropertySubType = ObjectProp->PropertyClass;
+	}
+
 	if (const UClassProperty* ClassProp = Cast<UClassProperty>(InProperty))
 	{
-		PropertySubType = ClassProp->PropertyClass;
+		PropertySubType = ClassProp->MetaClass;
+	}
+
+	if (const USoftClassProperty* ClassProp = Cast<USoftClassProperty>(InProperty))
+	{
+		PropertySubType = ClassProp->MetaClass;
 	}
 
 	if (const UStructProperty* StructProp = Cast<UStructProperty>(InProperty))
@@ -220,7 +230,7 @@ bool CalculatePropertyDef(PyTypeObject* InPyType, FPropertyDef& OutPropertyDef)
 {
 	if (PyObject_IsSubclass((PyObject*)InPyType, (PyObject*)&PyWrapperObjectType) == 1)
 	{
-		OutPropertyDef.PropertyClass = UClassProperty::StaticClass();
+		OutPropertyDef.PropertyClass = UObjectProperty::StaticClass();
 		OutPropertyDef.PropertySubType = (UObject*)FPyWrapperObjectMetaData::GetClass(InPyType);
 		return true;
 	}
@@ -357,10 +367,24 @@ UProperty* CreateProperty(const FPropertyDef& InPropertyDef, const int32 InArray
 	{
 		Prop->ArrayDim = InArrayDim;
 
+		if (UObjectPropertyBase* ObjectProp = Cast<UObjectPropertyBase>(Prop))
+		{
+			UClass* ClassType = CastChecked<UClass>(InPropertyDef.PropertySubType);
+			ObjectProp->SetPropertyClass(ClassType);
+		}
+
 		if (UClassProperty* ClassProp = Cast<UClassProperty>(Prop))
 		{
 			UClass* ClassType = CastChecked<UClass>(InPropertyDef.PropertySubType);
-			ClassProp->SetPropertyClass(ClassType);
+			ClassProp->SetPropertyClass(UClass::StaticClass());
+			ClassProp->SetMetaClass(ClassType);
+		}
+
+		if (USoftClassProperty* ClassProp = Cast<USoftClassProperty>(Prop))
+		{
+			UClass* ClassType = CastChecked<UClass>(InPropertyDef.PropertySubType);
+			ClassProp->SetPropertyClass(UClass::StaticClass());
+			ClassProp->SetMetaClass(ClassType);
 		}
 
 		if (UStructProperty* StructProp = Cast<UStructProperty>(Prop))
