Fixed crash when re-pointing an object referenced in Python

diff --git a/Engine/Plugins/Experimental/PythonScriptPlugin/Source/PythonScriptPlugin/Private/PyWrapperObject.cpp b/Engine/Plugins/Experimental/PythonScriptPlugin/Source/PythonScriptPlugin/Private/PyWrapperObject.cpp
--- a/Engine/Plugins/Experimental/PythonScriptPlugin/Source/PythonScriptPlugin/Private/PyWrapperObject.cpp
+++ b/Engine/Plugins/Experimental/PythonScriptPlugin/Source/PythonScriptPlugin/Private/PyWrapperObject.cpp
@@ -895,7 +895,15 @@ void FPyWrapperObjectMetaData::AddReferencedObjects(FPyWrapperBase* Instance, FR
 
 	UObject* OldInstance = Self->ObjectInstance;
 	Collector.AddReferencedObject(Self->ObjectInstance);
-	if (Self->ObjectInstance != OldInstance && Self->ObjectInstance != nullptr)
+
+	// Update the wrapped instance in the object factory
+	if (Self->ObjectInstance != OldInstance)
+	{
+		FPyWrapperObjectFactory::Get().UnmapInstance(OldInstance, Py_TYPE(Self));
+	}
+
+	// Update the object type
+	if (Self->ObjectInstance != OldInstance && Self->ObjectInstance)
 	{
 		// Object instance has been re-pointed, make sure we're still the correct type
 		PyTypeObject* NewPyType = FPyWrapperTypeRegistry::Get().GetWrappedClassType(Self->ObjectInstance->GetClass());
@@ -909,6 +917,12 @@ void FPyWrapperObjectMetaData::AddReferencedObjects(FPyWrapperBase* Instance, FR
 		}
 	}
 
+	// Update the wrapped instance in the object factory
+	if (Self->ObjectInstance != OldInstance && Self->ObjectInstance)
+	{
+		FPyWrapperObjectFactory::Get().MapInstance(Self->ObjectInstance, Self);
+	}
+
 	// We also need to ARO delegates on this object to catch ones that are wrapping Python callables (also recursing into nested structs and containers)
 	if (Self->ObjectInstance)
 	{
