Fixed Python stub file excluding all generated glue code when skipping doc generation

diff --git a/Engine/Plugins/Experimental/PythonScriptPlugin/Source/PythonScriptPlugin/Private/PyWrapperTypeRegistry.cpp b/Engine/Plugins/Experimental/PythonScriptPlugin/Source/PythonScriptPlugin/Private/PyWrapperTypeRegistry.cpp
--- a/Engine/Plugins/Experimental/PythonScriptPlugin/Source/PythonScriptPlugin/Private/PyWrapperTypeRegistry.cpp
+++ b/Engine/Plugins/Experimental/PythonScriptPlugin/Source/PythonScriptPlugin/Private/PyWrapperTypeRegistry.cpp
@@ -2093,14 +2093,12 @@ void FPyWrapperTypeRegistry::GenerateStubCodeForWrappedTypes(const EPyOnlineDocs
 	// which module it came from and where that module exists on disk.
 	auto ProcessWrappedDataArray = [this, &PythonScript, InDocGenFlags](const TMap<FName, PyTypeObject*>& WrappedData, const TSharedPtr<FPyOnlineDocsSection>& OnlineDocsSection)
 	{
-		if (InDocGenFlags == EPyOnlineDocsFilterFlags::IncludeNone)
+		if (OnlineDocsSection.IsValid())
 		{
-			return;
+			UE_LOG(LogPython, Display, TEXT("  ...generating Python API: %s"), *OnlineDocsSection->GetName());
+			PythonScript.WriteLine(FString::Printf(TEXT("##### %s #####"), *OnlineDocsSection->GetName()));
+			PythonScript.WriteNewLine();
 		}
-
-		UE_LOG(LogPython, Display, TEXT("  ...generating Python API: %s"), *OnlineDocsSection->GetName());
-		PythonScript.WriteLine(FString::Printf(TEXT("##### %s #####"), *OnlineDocsSection->GetName()));
-		PythonScript.WriteNewLine();
 		
 		FString ProjectTopDir;
 		if (FPaths::IsProjectFilePathSet())
@@ -2112,7 +2110,7 @@ void FPyWrapperTypeRegistry::GenerateStubCodeForWrappedTypes(const EPyOnlineDocs
 		{
 			TSharedPtr<PyGenUtil::FGeneratedWrappedType> GeneratedWrappedType = GeneratedWrappedTypes.FindRef(WrappedDataPair.Key);
 
-			if ((InDocGenFlags != EPyOnlineDocsFilterFlags::IncludeAll) && GeneratedWrappedType.IsValid())
+			if ((InDocGenFlags != EPyOnlineDocsFilterFlags::IncludeAll) && GeneratedWrappedType.IsValid() && OnlineDocsSection.IsValid())
 			{
 				const UField* MetaType = GeneratedWrappedType->MetaData->GetMetaType();
 
