Fixed Python recent files going to wrong config path

diff --git a/Engine/Plugins/Experimental/PythonScriptPlugin/Source/PythonScriptPlugin/Private/PythonScriptPlugin.cpp b/Engine/Plugins/Experimental/PythonScriptPlugin/Source/PythonScriptPlugin/Private/PythonScriptPlugin.cpp
--- a/Engine/Plugins/Experimental/PythonScriptPlugin/Source/PythonScriptPlugin/Private/PythonScriptPlugin.cpp
+++ b/Engine/Plugins/Experimental/PythonScriptPlugin/Source/PythonScriptPlugin/Private/PythonScriptPlugin.cpp
@@ -141,7 +141,6 @@ public:
 	FPythonCommandMenuImpl()
 		: bRecentsFilesDirty(false)
 	{
-		ConfigFilename = UObject::StaticClass()->GetDefaultObject()->GetGlobalUserConfigFilename();
 	}
 
 	virtual void OnStartupMenu() override
@@ -215,9 +214,9 @@ private:
 	{
 		RecentsFiles.Reset();
 
-		GConfig->GetString(STR_ConfigSection, STR_ConfigDirectoryKey, LastDirectory, ConfigFilename);
+		GConfig->GetString(STR_ConfigSection, STR_ConfigDirectoryKey, LastDirectory, GEditorPerProjectIni);
 
-		FConfigSection* Sec = GConfig->GetSectionPrivate(STR_ConfigSection, false, true, ConfigFilename);
+		FConfigSection* Sec = GConfig->GetSectionPrivate(STR_ConfigSection, false, true, GEditorPerProjectIni);
 		if (Sec)
 		{
 			TArray<FConfigValue> List;
@@ -235,9 +234,9 @@ private:
 
 	void SaveConfig() const
 	{
-		GConfig->SetString(STR_ConfigSection, STR_ConfigDirectoryKey, *LastDirectory, ConfigFilename);
+		GConfig->SetString(STR_ConfigSection, STR_ConfigDirectoryKey, *LastDirectory, GEditorPerProjectIni);
 
-		FConfigSection* Sec = GConfig->GetSectionPrivate(STR_ConfigSection, true, false, ConfigFilename);
+		FConfigSection* Sec = GConfig->GetSectionPrivate(STR_ConfigSection, true, false, GEditorPerProjectIni);
 		if (Sec)
 		{
 			Sec->Remove(NAME_ConfigRecentsFilesyKey);
@@ -329,7 +328,6 @@ private:
 	TArray<FString> RecentsFiles;
 	FString LastDirectory;
 
-	FString ConfigFilename;
 	bool bRecentsFilesDirty;
 };
 #endif // WITH_EDITOR
