Merge change from Dev-Editor to allow editing properties marked with CPF_DisableEditOnInstance on asset instances

#RNX

diff --git a/Engine/Plugins/Experimental/PythonScriptPlugin/Source/PythonScriptPlugin/Private/PyWrapperObject.cpp b/Engine/Plugins/Experimental/PythonScriptPlugin/Source/PythonScriptPlugin/Private/PyWrapperObject.cpp
--- a/Engine/Plugins/Experimental/PythonScriptPlugin/Source/PythonScriptPlugin/Private/PyWrapperObject.cpp
+++ b/Engine/Plugins/Experimental/PythonScriptPlugin/Source/PythonScriptPlugin/Private/PyWrapperObject.cpp
@@ -132,7 +132,7 @@ int FPyWrapperObject::SetPropertyValue(FPyWrapperObject* InSelf, PyObject* InVal
 	}
 
 	const FPyWrapperOwnerContext ChangeOwner = InNotifyChange ? FPyWrapperOwnerContext((PyObject*)InSelf, InPropDef.Prop) : FPyWrapperOwnerContext();
-	return PyGenUtil::SetPropertyValue(InSelf->ObjectInstance->GetClass(), InSelf->ObjectInstance, InValue, InPropDef, InPythonAttrName, ChangeOwner, InReadOnlyFlags, InSelf->ObjectInstance->IsTemplate(), *PyUtil::GetErrorContext(InSelf));
+	return PyGenUtil::SetPropertyValue(InSelf->ObjectInstance->GetClass(), InSelf->ObjectInstance, InValue, InPropDef, InPythonAttrName, ChangeOwner, InReadOnlyFlags, InSelf->ObjectInstance->IsTemplate() || InSelf->ObjectInstance->IsAsset(), *PyUtil::GetErrorContext(InSelf));
 }
 
 PyObject* FPyWrapperObject::CallGetterFunction(FPyWrapperObject* InSelf, const PyGenUtil::FGeneratedWrappedFunction& InFuncDef)
diff --git a/Engine/Plugins/Experimental/PythonScriptPlugin/Source/PythonScriptPlugin/Private/PyWrapperStruct.cpp b/Engine/Plugins/Experimental/PythonScriptPlugin/Source/PythonScriptPlugin/Private/PyWrapperStruct.cpp
--- a/Engine/Plugins/Experimental/PythonScriptPlugin/Source/PythonScriptPlugin/Private/PyWrapperStruct.cpp
+++ b/Engine/Plugins/Experimental/PythonScriptPlugin/Source/PythonScriptPlugin/Private/PyWrapperStruct.cpp
@@ -444,7 +444,7 @@ int FPyWrapperStruct::SetPropertyValue(FPyWrapperStruct* InSelf, PyObject* InVal
 	bool OwnerIsTemplate = false;
 	if (const UObject* OwnerObject = PyUtil::GetOwnerObject((PyObject*)InSelf))
 	{
-		OwnerIsTemplate = OwnerObject->IsTemplate();
+		OwnerIsTemplate = OwnerObject->IsTemplate() || OwnerObject->IsAsset();
 	}
 
 	const FPyWrapperOwnerContext ChangeOwner = InNotifyChange ? FPyWrapperOwnerContext((PyObject*)InSelf, InPropDef.Prop) : FPyWrapperOwnerContext();
