Ensure Python type generation cleans-up after itself correctly if the generation fails

diff --git a/Engine/Plugins/Experimental/PythonScriptPlugin/Source/PythonScriptPlugin/Private/PyCore.cpp b/Engine/Plugins/Experimental/PythonScriptPlugin/Source/PythonScriptPlugin/Private/PyCore.cpp
--- a/Engine/Plugins/Experimental/PythonScriptPlugin/Source/PythonScriptPlugin/Private/PyCore.cpp
+++ b/Engine/Plugins/Experimental/PythonScriptPlugin/Source/PythonScriptPlugin/Private/PyCore.cpp
@@ -1644,9 +1644,9 @@ PyMethodDef PyCoreMethods[] = {
 
 void InitializeModule()
 {
-	GPythonPropertyContainer.Reset(NewObject<UStruct>(GetTransientPackage(), TEXT("PythonProperties")));
+	GPythonPropertyContainer.Reset(NewObject<UStruct>(GetTransientPackage(), TEXT("PythonProperties"), RF_Transient));
 
-	GPythonTypeContainer.Reset(NewObject<UPackage>(nullptr, TEXT("/Engine/PythonTypes"), RF_Public));
+	GPythonTypeContainer.Reset(NewObject<UPackage>(nullptr, TEXT("/Engine/PythonTypes"), RF_Public | RF_Transient));
 	GPythonTypeContainer->SetPackageFlags(PKG_ContainsScript);
 
 	PyGenUtil::FNativePythonModule NativePythonModule;
diff --git a/Engine/Plugins/Experimental/PythonScriptPlugin/Source/PythonScriptPlugin/Private/PyWrapperEnum.cpp b/Engine/Plugins/Experimental/PythonScriptPlugin/Source/PythonScriptPlugin/Private/PyWrapperEnum.cpp
--- a/Engine/Plugins/Experimental/PythonScriptPlugin/Source/PythonScriptPlugin/Private/PyWrapperEnum.cpp
+++ b/Engine/Plugins/Experimental/PythonScriptPlugin/Source/PythonScriptPlugin/Private/PyWrapperEnum.cpp
@@ -714,102 +714,135 @@ bool FPyWrapperEnumMetaData::IsEnumFinalized(FPyWrapperEnum* Instance)
 	return IsEnumFinalized(Py_TYPE(Instance));
 }
 
-struct FPythonGeneratedEnumUtil
+class FPythonGeneratedEnumBuilder
 {
-	static UPythonGeneratedEnum* CreateEnum(const FString& InEnumName, UObject* InEnumOuter)
+public:
+	FPythonGeneratedEnumBuilder(const FString& InEnumName, PyTypeObject* InPyType)
+		: EnumName(InEnumName)
+		, PyType(InPyType)
+		, NewEnum(nullptr)
 	{
-		UPythonGeneratedEnum* Enum = NewObject<UPythonGeneratedEnum>(InEnumOuter, *InEnumName, RF_Public | RF_Standalone);
-		Enum->SetMetaData(TEXT("BlueprintType"), TEXT("true"));
-		return Enum;
+		UObject* EnumOuter = GetPythonTypeContainer();
+
+		// Enum instances are re-used if they already exist
+		NewEnum = FindObject<UPythonGeneratedEnum>(EnumOuter, *EnumName);
+		if (!NewEnum)
+		{
+			NewEnum = NewObject<UPythonGeneratedEnum>(EnumOuter, *EnumName, RF_Public | RF_Standalone | RF_Transient);
+			NewEnum->SetMetaData(TEXT("BlueprintType"), TEXT("true"));
+		}
+		NewEnum->EnumValueDefs.Reset();
 	}
 
-	static void FinalizeEnum(UPythonGeneratedEnum* InEnum, PyTypeObject* InPyType)
+	~FPythonGeneratedEnumBuilder()
 	{
+		// If NewEnum is still set at this point, if means Finalize wasn't called and we should destroy the partially built enum
+		if (NewEnum)
+		{
+			NewEnum->ClearFlags(RF_Public | RF_Standalone);
+			NewEnum = nullptr;
+
+			CollectGarbage(GARBAGE_COLLECTION_KEEPFLAGS);
+		}
+	}
+
+	UPythonGeneratedEnum* Finalize(const TArray<FPyUValueDef*>& InPyValueDefs)
+	{
+		// Populate the enum with its values, and replace the definitions with real descriptors
+		if (!RegisterDescriptors(InPyValueDefs))
+		{
+			return nullptr;
+		}
+
+		// Let Python know that we've changed its type
+		PyType_Modified(PyType);
+
 		// Finalize the enum
-		InEnum->Bind();
+		NewEnum->Bind();
 
 		// Add the object meta-data to the type
-		InEnum->PyMetaData.Enum = InEnum;
-		InEnum->PyMetaData.bFinalized = true;
-		FPyWrapperEnumMetaData::SetMetaData(InPyType, &InEnum->PyMetaData);
+		NewEnum->PyMetaData.Enum = NewEnum;
+		NewEnum->PyMetaData.bFinalized = true;
+		FPyWrapperEnumMetaData::SetMetaData(PyType, &NewEnum->PyMetaData);
 
 		// Map the Unreal enum to the Python type
-		InEnum->PyType = FPyTypeObjectPtr::NewReference(InPyType);
-		FPyWrapperTypeRegistry::Get().RegisterWrappedEnumType(InEnum->GetFName(), InPyType);
+		NewEnum->PyType = FPyTypeObjectPtr::NewReference(PyType);
+		FPyWrapperTypeRegistry::Get().RegisterWrappedEnumType(NewEnum->GetFName(), PyType);
+
+		// Null the NewEnum pointer so the destructor doesn't kill it
+		UPythonGeneratedEnum* FinalizedEnum = NewEnum;
+		NewEnum = nullptr;
+		return FinalizedEnum;
 	}
 
-	static bool CreateValueFromDefinition(UPythonGeneratedEnum* InEnum, PyTypeObject* InPyType, const FString& InFieldName, FPyUValueDef* InPyValueDef)
+	bool CreateValueFromDefinition(const FString& InFieldName, FPyUValueDef* InPyValueDef)
 	{
 		int64 EnumValue = 0;
 		if (!PyConversion::Nativize(InPyValueDef->Value, EnumValue))
 		{
-			PyUtil::SetPythonError(PyExc_TypeError, InPyType, *FString::Printf(TEXT("Failed to convert enum value for '%s'"), *InFieldName));
+			PyUtil::SetPythonError(PyExc_TypeError, PyType, *FString::Printf(TEXT("Failed to convert enum value for '%s'"), *InFieldName));
 			return false;
 		}
 
 		// Build the definition data for the new enum value
-		UPythonGeneratedEnum::FEnumValueDef& EnumValueDef = *InEnum->EnumValueDefs.Add_GetRef(MakeShared<UPythonGeneratedEnum::FEnumValueDef>());
+		UPythonGeneratedEnum::FEnumValueDef& EnumValueDef = *NewEnum->EnumValueDefs.Add_GetRef(MakeShared<UPythonGeneratedEnum::FEnumValueDef>());
 		EnumValueDef.Value = EnumValue;
 		EnumValueDef.Name = InFieldName;
 
 		return true;
 	}
 
-	static bool RegisterDescriptors(UPythonGeneratedEnum* InEnum, PyTypeObject* InPyType, const TArray<FPyUValueDef*>& InPyValueDefs)
+private:
+	bool RegisterDescriptors(const TArray<FPyUValueDef*>& InPyValueDefs)
 	{
 		// Populate the enum with its values
-		check(InPyValueDefs.Num() == InEnum->EnumValueDefs.Num());
+		check(InPyValueDefs.Num() == NewEnum->EnumValueDefs.Num());
 		{
-			const FString EnumName = InEnum->GetName();
-
 			TArray<TPair<FName, int64>> ValueNames;
-			for (const TSharedPtr<UPythonGeneratedEnum::FEnumValueDef>& EnumValueDef : InEnum->EnumValueDefs)
+			for (const TSharedPtr<UPythonGeneratedEnum::FEnumValueDef>& EnumValueDef : NewEnum->EnumValueDefs)
 			{
 				const FString NamespacedValueName = FString::Printf(TEXT("%s::%s"), *EnumName, *EnumValueDef->Name);
 				ValueNames.Emplace(MakeTuple(FName(*NamespacedValueName), EnumValueDef->Value));
 			}
-			if (!InEnum->SetEnums(ValueNames, UEnum::ECppForm::Namespaced))
+			if (!NewEnum->SetEnums(ValueNames, UEnum::ECppForm::Namespaced))
 			{
-				PyUtil::SetPythonError(PyExc_Exception, InPyType, TEXT("Failed to set enum values"));
+				PyUtil::SetPythonError(PyExc_Exception, PyType, TEXT("Failed to set enum values"));
 				return false;
 			}
 
 			// Can't set the meta-data until SetEnums has been called
 			for (int32 EnumEntryIndex = 0; EnumEntryIndex < InPyValueDefs.Num(); ++EnumEntryIndex)
 			{
-				FPyUValueDef::ApplyMetaData(InPyValueDefs[EnumEntryIndex], [InEnum, EnumEntryIndex](const FString& InMetaDataKey, const FString& InMetaDataValue)
+				FPyUValueDef::ApplyMetaData(InPyValueDefs[EnumEntryIndex], [this, EnumEntryIndex](const FString& InMetaDataKey, const FString& InMetaDataValue)
 				{
-					InEnum->SetMetaData(*InMetaDataKey, *InMetaDataValue, EnumEntryIndex);
+					NewEnum->SetMetaData(*InMetaDataKey, *InMetaDataValue, EnumEntryIndex);
 				});
-				InEnum->EnumValueDefs[EnumEntryIndex]->DocString = PyGenUtil::GetEnumEntryTooltip(InEnum, EnumEntryIndex);
+				NewEnum->EnumValueDefs[EnumEntryIndex]->DocString = PyGenUtil::GetEnumEntryTooltip(NewEnum, EnumEntryIndex);
 			}
 		}
 
 		// Replace the definitions with real descriptors
-		for (const TSharedPtr<UPythonGeneratedEnum::FEnumValueDef>& EnumValueDef : InEnum->EnumValueDefs)
+		for (const TSharedPtr<UPythonGeneratedEnum::FEnumValueDef>& EnumValueDef : NewEnum->EnumValueDefs)
 		{
-			FPyWrapperEnum* PyEnumEntry = FPyWrapperEnum::AddEnumEntry(InPyType, EnumValueDef->Value, TCHAR_TO_UTF8(*EnumValueDef->Name), TCHAR_TO_UTF8(*EnumValueDef->DocString));
+			FPyWrapperEnum* PyEnumEntry = FPyWrapperEnum::AddEnumEntry(PyType, EnumValueDef->Value, TCHAR_TO_UTF8(*EnumValueDef->Name), TCHAR_TO_UTF8(*EnumValueDef->DocString));
 			if (PyEnumEntry)
 			{
-				InEnum->PyMetaData.EnumEntries.Add(PyEnumEntry);
+				NewEnum->PyMetaData.EnumEntries.Add(PyEnumEntry);
 			}
 		}
 
 		return true;
 	}
+
+	FString EnumName;
+	PyTypeObject* PyType;
+	UPythonGeneratedEnum* NewEnum;
 };
 
 UPythonGeneratedEnum* UPythonGeneratedEnum::GenerateEnum(PyTypeObject* InPyType)
 {
-	UObject* EnumOuter = GetPythonTypeContainer();
-	const FString EnumName = PyUtil::GetCleanTypename(InPyType);
-
-	UPythonGeneratedEnum* Enum = FindObject<UPythonGeneratedEnum>(EnumOuter, *EnumName);
-	if (!Enum)
-	{
-		Enum = FPythonGeneratedEnumUtil::CreateEnum(EnumName, EnumOuter);
-	}
-	Enum->EnumValueDefs.Reset();
+	// Builder used to generate the enum
+	FPythonGeneratedEnumBuilder PythonEnumBuilder(PyUtil::GetCleanTypename(InPyType), InPyType);
 
 	// Add the values to this enum
 	TArray<FPyUValueDef*> PyValueDefs;
@@ -826,7 +859,7 @@ UPythonGeneratedEnum* UPythonGeneratedEnum::GenerateEnum(PyTypeObject* InPyType)
 				FPyUValueDef* PyValueDef = (FPyUValueDef*)FieldValue;
 				PyValueDefs.Add(PyValueDef);
 
-				if (!FPythonGeneratedEnumUtil::CreateValueFromDefinition(Enum, InPyType, FieldName, PyValueDef))
+				if (!PythonEnumBuilder.CreateValueFromDefinition(FieldName, PyValueDef))
 				{
 					return nullptr;
 				}
@@ -848,19 +881,8 @@ UPythonGeneratedEnum* UPythonGeneratedEnum::GenerateEnum(PyTypeObject* InPyType)
 		}
 	}
 
-	// Populate the enum with its values, and replace the definitions with real descriptors
-	if (!FPythonGeneratedEnumUtil::RegisterDescriptors(Enum, InPyType, PyValueDefs))
-	{
-		return nullptr;
-	}
-
-	// Let Python know that we've changed its type
-	PyType_Modified(InPyType);
-
-	// Finalize the enum
-	FPythonGeneratedEnumUtil::FinalizeEnum(Enum, InPyType);
-
-	return Enum;
+	// Finalize the struct with its value meta-data
+	return PythonEnumBuilder.Finalize(PyValueDefs);
 }
 
 #endif	// WITH_PYTHON
diff --git a/Engine/Plugins/Experimental/PythonScriptPlugin/Source/PythonScriptPlugin/Private/PyWrapperEnum.h b/Engine/Plugins/Experimental/PythonScriptPlugin/Source/PythonScriptPlugin/Private/PyWrapperEnum.h
--- a/Engine/Plugins/Experimental/PythonScriptPlugin/Source/PythonScriptPlugin/Private/PyWrapperEnum.h
+++ b/Engine/Plugins/Experimental/PythonScriptPlugin/Source/PythonScriptPlugin/Private/PyWrapperEnum.h
@@ -143,7 +143,7 @@ private:
 	/** Meta-data for this generated enum that is applied to the Python type */
 	FPyWrapperEnumMetaData PyMetaData;
 
-	friend struct FPythonGeneratedEnumUtil;
+	friend class FPythonGeneratedEnumBuilder;
 
 #endif	// WITH_PYTHON
 };
diff --git a/Engine/Plugins/Experimental/PythonScriptPlugin/Source/PythonScriptPlugin/Private/PyWrapperObject.cpp b/Engine/Plugins/Experimental/PythonScriptPlugin/Source/PythonScriptPlugin/Private/PyWrapperObject.cpp
--- a/Engine/Plugins/Experimental/PythonScriptPlugin/Source/PythonScriptPlugin/Private/PyWrapperObject.cpp
+++ b/Engine/Plugins/Experimental/PythonScriptPlugin/Source/PythonScriptPlugin/Private/PyWrapperObject.cpp
@@ -1067,70 +1067,119 @@ bool FPyWrapperObjectMetaData::IsClassDeprecated(FPyWrapperObject* Instance, FSt
 	return IsClassDeprecated(Py_TYPE(Instance), OutDeprecationMessage);
 }
 
-struct FPythonGeneratedClassUtil
+class FPythonGeneratedClassBuilder
 {
-	static void PrepareOldClassForReinstancing(UPythonGeneratedClass* InOldClass)
+public:
+	FPythonGeneratedClassBuilder(const FString& InClassName, UClass* InSuperClass, PyTypeObject* InPyType)
+		: ClassName(InClassName)
+		, PyType(InPyType)
+		, OldClass(nullptr)
+		, NewClass(nullptr)
 	{
-		const FString OldClassName = MakeUniqueObjectName(InOldClass->GetOuter(), InOldClass->GetClass(), *FString::Printf(TEXT("%s_REINST"), *InOldClass->GetName())).ToString();
-		InOldClass->ClassFlags |= CLASS_NewerVersionExists;
-		InOldClass->SetFlags(RF_NewerVersionExists);
-		InOldClass->ClearFlags(RF_Public | RF_Standalone);
-		InOldClass->Rename(*OldClassName, nullptr, REN_DontCreateRedirectors);
+		UObject* ClassOuter = GetPythonTypeContainer();
+
+		// Find any existing class with the name we want to use
+		OldClass = FindObject<UPythonGeneratedClass>(ClassOuter, *ClassName);
+
+		// Create a new class with a temporary name; we will rename it as part of Finalize
+		const FString NewClassName = MakeUniqueObjectName(ClassOuter, UPythonGeneratedClass::StaticClass(), *FString::Printf(TEXT("%s_NEWINST"), *ClassName)).ToString();
+		NewClass = NewObject<UPythonGeneratedClass>(ClassOuter, *NewClassName, RF_Public | RF_Standalone | RF_Transient);
+		NewClass->SetMetaData(TEXT("BlueprintType"), TEXT("true"));
+		NewClass->SetSuperStruct(InSuperClass);
 	}
 
-	static UPythonGeneratedClass* CreateClass(const FString& InClassName, UObject* InClassOuter, UClass* InSuperClass)
+	~FPythonGeneratedClassBuilder()
 	{
-		UPythonGeneratedClass* Class = NewObject<UPythonGeneratedClass>(InClassOuter, *InClassName, RF_Public | RF_Standalone);
-		Class->SetMetaData(TEXT("BlueprintType"), TEXT("true"));
-		Class->SetSuperStruct(InSuperClass);
-		return Class;
+		// If NewClass is still set at this point, if means Finalize wasn't called and we should destroy the partially built class
+		if (NewClass)
+		{
+			NewClass->ClearFlags(RF_Public | RF_Standalone);
+			NewClass = nullptr;
+
+			CollectGarbage(GARBAGE_COLLECTION_KEEPFLAGS);
+		}
 	}
 
-	static void FinalizeClass(UPythonGeneratedClass* InClass, PyTypeObject* InPyType)
+	UPythonGeneratedClass* Finalize(FPyObjectPtr InPyPostInitFunction)
 	{
+		// Set the post-init function
+		NewClass->PyPostInitFunction = InPyPostInitFunction;
+		if (!NewClass->PyPostInitFunction)
+		{
+			return nullptr;
+		}
+
+		// Replace the definitions with real descriptors
+		if (!RegisterDescriptors())
+		{
+			return nullptr;
+		}
+
+		// Let Python know that we've changed its type
+		PyType_Modified(PyType);
+
+		// We can no longer fail, so prepare the old class for removal and set the correct name on the new class
+		if (OldClass)
+		{
+			PrepareOldClassForReinstancing();
+		}
+		NewClass->Rename(*ClassName, nullptr, REN_DontCreateRedirectors);
+
 		// Finalize the class
-		InClass->Bind();
-		InClass->StaticLink(true);
-		InClass->AssembleReferenceTokenStream();
+		NewClass->Bind();
+		NewClass->StaticLink(true);
+		NewClass->AssembleReferenceTokenStream();
 
 		// Add the object meta-data to the type
-		InClass->PyMetaData.Class = InClass;
-		FPyWrapperObjectMetaData::SetMetaData(InPyType, &InClass->PyMetaData);
+		NewClass->PyMetaData.Class = NewClass;
+		FPyWrapperObjectMetaData::SetMetaData(PyType, &NewClass->PyMetaData);
 
 		// Map the Unreal class to the Python type
-		InClass->PyType = FPyTypeObjectPtr::NewReference(InPyType);
-		FPyWrapperTypeRegistry::Get().RegisterWrappedClassType(InClass->GetFName(), InPyType);
+		NewClass->PyType = FPyTypeObjectPtr::NewReference(PyType);
+		FPyWrapperTypeRegistry::Get().RegisterWrappedClassType(NewClass->GetFName() , PyType);
 
 		// Ensure the CDO exists
-		InClass->GetDefaultObject();
+		NewClass->GetDefaultObject();
+
+		// Re-instance the old class and re-parent any derived classes to this new type
+		if (OldClass)
+		{
+			FPyWrapperTypeReinstancer::Get().AddPendingClass(OldClass, NewClass);
+			UPythonGeneratedClass::ReparentDerivedClasses(OldClass, NewClass);
+		}	
+
+		// Null the NewClass pointer so the destructor doesn't kill it
+		UPythonGeneratedClass* FinalizedClass = NewClass;
+		NewClass = nullptr;
+		return FinalizedClass;
 	}
 
-	static bool CreatePropertyFromDefinition(UPythonGeneratedClass* InClass, PyTypeObject* InPyType, const FString& InFieldName, FPyUPropertyDef* InPyPropDef)
+	bool CreatePropertyFromDefinition(const FString& InFieldName, FPyUPropertyDef* InPyPropDef)
 	{
-		UClass* SuperClass = InClass->GetSuperClass();
+		UClass* SuperClass = NewClass->GetSuperClass();
 
 		// Resolve the property name to match any previously exported properties from the parent type
-		const FName PropName = FPyWrapperObjectMetaData::ResolvePropertyName(InPyType->tp_base, *InFieldName);
+		const FName PropName = FPyWrapperObjectMetaData::ResolvePropertyName(PyType->tp_base, *InFieldName);
 		if (SuperClass->FindPropertyByName(PropName))
 		{
-			PyUtil::SetPythonError(PyExc_Exception, InPyType, *FString::Printf(TEXT("Property '%s' (%s) cannot override a property from the base type"), *InFieldName, *PyUtil::GetFriendlyTypename(InPyPropDef->PropType)));
+			PyUtil::SetPythonError(PyExc_Exception, PyType, *FString::Printf(TEXT("Property '%s' (%s) cannot override a property from the base type"), *InFieldName, *PyUtil::GetFriendlyTypename(InPyPropDef->PropType)));
 			return false;
 		}
 
 		// Create the property from its definition
-		UProperty* Prop = PyUtil::CreateProperty(InPyPropDef->PropType, 1, InClass, PropName);
+		UProperty* Prop = PyUtil::CreateProperty(InPyPropDef->PropType, 1, NewClass, PropName);
 		if (!Prop)
 		{
-			PyUtil::SetPythonError(PyExc_Exception, InPyType, *FString::Printf(TEXT("Failed to create property for '%s' (%s)"), *InFieldName, *PyUtil::GetFriendlyTypename(InPyPropDef->PropType)));
+			PyUtil::SetPythonError(PyExc_Exception, PyType, *FString::Printf(TEXT("Failed to create property for '%s' (%s)"), *InFieldName, *PyUtil::GetFriendlyTypename(InPyPropDef->PropType)));
 			return false;
 		}
 		Prop->PropertyFlags |= (CPF_Edit | CPF_BlueprintVisible);
 		FPyUPropertyDef::ApplyMetaData(InPyPropDef, Prop);
-		InClass->AddCppProperty(Prop);
+		NewClass->AddCppProperty(Prop);
 
 		// Resolve any getter/setter function names
-		const FName GetterFuncName = FPyWrapperObjectMetaData::ResolveFunctionName(InPyType->tp_base, *InPyPropDef->GetterFuncName);
-		const FName SetterFuncName = FPyWrapperObjectMetaData::ResolveFunctionName(InPyType->tp_base, *InPyPropDef->SetterFuncName);
+		const FName GetterFuncName = FPyWrapperObjectMetaData::ResolveFunctionName(PyType->tp_base, *InPyPropDef->GetterFuncName);
+		const FName SetterFuncName = FPyWrapperObjectMetaData::ResolveFunctionName(PyType->tp_base, *InPyPropDef->SetterFuncName);
 		if (!GetterFuncName.IsNone())
 		{
 			Prop->SetMetaData(PyGenUtil::BlueprintGetterMetaDataKey, *GetterFuncName.ToString());
@@ -1141,12 +1190,12 @@ struct FPythonGeneratedClassUtil
 		}
 
 		// Build the definition data for the new property accessor
-		PyGenUtil::FPropertyDef& PropDef = *InClass->PropertyDefs.Add_GetRef(MakeShared<PyGenUtil::FPropertyDef>());
+		PyGenUtil::FPropertyDef& PropDef = *NewClass->PropertyDefs.Add_GetRef(MakeShared<PyGenUtil::FPropertyDef>());
 		PropDef.GeneratedWrappedGetSet.GetSetName = PyGenUtil::TCHARToUTF8Buffer(*InFieldName);
 		PropDef.GeneratedWrappedGetSet.GetSetDoc = PyGenUtil::TCHARToUTF8Buffer(*FString::Printf(TEXT("type: %s\n%s"), *PyGenUtil::GetPropertyPythonType(Prop), *PyGenUtil::GetFieldTooltip(Prop)));
 		PropDef.GeneratedWrappedGetSet.Prop.SetProperty(Prop);
-		PropDef.GeneratedWrappedGetSet.GetFunc.SetFunction(InClass->FindFunctionByName(GetterFuncName));
-		PropDef.GeneratedWrappedGetSet.SetFunc.SetFunction(InClass->FindFunctionByName(SetterFuncName));
+		PropDef.GeneratedWrappedGetSet.GetFunc.SetFunction(NewClass->FindFunctionByName(GetterFuncName));
+		PropDef.GeneratedWrappedGetSet.SetFunc.SetFunction(NewClass->FindFunctionByName(SetterFuncName));
 		PropDef.GeneratedWrappedGetSet.GetCallback = (getter)&FPyWrapperObject::Getter_Impl;
 		PropDef.GeneratedWrappedGetSet.SetCallback = (setter)&FPyWrapperObject::Setter_Impl;
 		PropDef.GeneratedWrappedGetSet.ToPython(PropDef.PyGetSet);
@@ -1154,7 +1203,7 @@ struct FPythonGeneratedClassUtil
 		// If this property has a getter or setter, also make an internal version with the get/set function cleared so that Python can read/write the internal property value
 		if (PropDef.GeneratedWrappedGetSet.GetFunc.Func || PropDef.GeneratedWrappedGetSet.SetFunc.Func)
 		{
-			PyGenUtil::FPropertyDef& InternalPropDef = *InClass->PropertyDefs.Add_GetRef(MakeShared<PyGenUtil::FPropertyDef>());
+			PyGenUtil::FPropertyDef& InternalPropDef = *NewClass->PropertyDefs.Add_GetRef(MakeShared<PyGenUtil::FPropertyDef>());
 			InternalPropDef.GeneratedWrappedGetSet.GetSetName = PyGenUtil::TCHARToUTF8Buffer(*FString::Printf(TEXT("_%s"), *InFieldName));
 			InternalPropDef.GeneratedWrappedGetSet.GetSetDoc = PropDef.GeneratedWrappedGetSet.GetSetDoc;
 			InternalPropDef.GeneratedWrappedGetSet.Prop.SetProperty(Prop);
@@ -1166,61 +1215,61 @@ struct FPythonGeneratedClassUtil
 		return true;
 	}
 
-	static bool CreateFunctionFromDefinition(UPythonGeneratedClass* InClass, PyTypeObject* InPyType, const FString& InFieldName, FPyUFunctionDef* InPyFuncDef)
+	bool CreateFunctionFromDefinition(const FString& InFieldName, FPyUFunctionDef* InPyFuncDef)
 	{
-		UClass* SuperClass = InClass->GetSuperClass();
+		UClass* SuperClass = NewClass->GetSuperClass();
 
 		// Validate the function definition makes sense
 		if (EnumHasAllFlags(InPyFuncDef->FuncFlags, EPyUFunctionDefFlags::Override))
 		{
 			if (EnumHasAnyFlags(InPyFuncDef->FuncFlags, EPyUFunctionDefFlags::Static | EPyUFunctionDefFlags::Getter | EPyUFunctionDefFlags::Setter))
 			{
-				PyUtil::SetPythonError(PyExc_Exception, InPyType, *FString::Printf(TEXT("Method '%s' specified as 'override' cannot also specify 'static', 'getter', or 'setter'"), *InFieldName));
+				PyUtil::SetPythonError(PyExc_Exception, PyType, *FString::Printf(TEXT("Method '%s' specified as 'override' cannot also specify 'static', 'getter', or 'setter'"), *InFieldName));
 				return false;
 			}
 			if (InPyFuncDef->FuncRetType != Py_None || InPyFuncDef->FuncParamTypes != Py_None)
 			{
-				PyUtil::SetPythonError(PyExc_Exception, InPyType, *FString::Printf(TEXT("Method '%s' specified as 'override' cannot also specify 'ret' or 'params'"), *InFieldName));
+				PyUtil::SetPythonError(PyExc_Exception, PyType, *FString::Printf(TEXT("Method '%s' specified as 'override' cannot also specify 'ret' or 'params'"), *InFieldName));
 				return false;
 			}
 		}
 		if (EnumHasAllFlags(InPyFuncDef->FuncFlags, EPyUFunctionDefFlags::Static) && EnumHasAnyFlags(InPyFuncDef->FuncFlags, EPyUFunctionDefFlags::Getter | EPyUFunctionDefFlags::Setter))
 		{
-			PyUtil::SetPythonError(PyExc_Exception, InPyType, *FString::Printf(TEXT("Method '%s' specified as 'static' cannot also specify 'getter' or 'setter'"), *InFieldName));
+			PyUtil::SetPythonError(PyExc_Exception, PyType, *FString::Printf(TEXT("Method '%s' specified as 'static' cannot also specify 'getter' or 'setter'"), *InFieldName));
 			return false;
 		}
 		if (EnumHasAllFlags(InPyFuncDef->FuncFlags, EPyUFunctionDefFlags::Getter))
 		{
 			if (EnumHasAnyFlags(InPyFuncDef->FuncFlags, EPyUFunctionDefFlags::Setter))
 			{
-				PyUtil::SetPythonError(PyExc_Exception, InPyType, *FString::Printf(TEXT("Method '%s' specified as 'getter' cannot also specify 'setter'"), *InFieldName));
+				PyUtil::SetPythonError(PyExc_Exception, PyType, *FString::Printf(TEXT("Method '%s' specified as 'getter' cannot also specify 'setter'"), *InFieldName));
 				return false;
 			}
 			if (EnumHasAnyFlags(InPyFuncDef->FuncFlags, EPyUFunctionDefFlags::Impure))
 			{
-				PyUtil::SetPythonError(PyExc_Exception, InPyType, *FString::Printf(TEXT("Method '%s' specified as 'getter' must also specify 'pure=True'"), *InFieldName));
+				PyUtil::SetPythonError(PyExc_Exception, PyType, *FString::Printf(TEXT("Method '%s' specified as 'getter' must also specify 'pure=True'"), *InFieldName));
 				return false;
 			}
 		}
 
 		// Resolve the function name to match any previously exported functions from the parent type
-		const FName FuncName = FPyWrapperObjectMetaData::ResolveFunctionName(InPyType->tp_base, *InFieldName);
+		const FName FuncName = FPyWrapperObjectMetaData::ResolveFunctionName(PyType->tp_base, *InFieldName);
 		const UFunction* SuperFunc = SuperClass->FindFunctionByName(FuncName);
 		if (SuperFunc && !EnumHasAllFlags(InPyFuncDef->FuncFlags, EPyUFunctionDefFlags::Override))
 		{
-			PyUtil::SetPythonError(PyExc_Exception, InPyType, *FString::Printf(TEXT("Method '%s' cannot override a method from the base type (did you forget to specify 'override=True'?)"), *InFieldName));
+			PyUtil::SetPythonError(PyExc_Exception, PyType, *FString::Printf(TEXT("Method '%s' cannot override a method from the base type (did you forget to specify 'override=True'?)"), *InFieldName));
 			return false;
 		}
 		if (EnumHasAllFlags(InPyFuncDef->FuncFlags, EPyUFunctionDefFlags::Override))
 		{
 			if (!SuperFunc)
 			{
-				PyUtil::SetPythonError(PyExc_Exception, InPyType, *FString::Printf(TEXT("Method '%s' was set to 'override', but no method was found to override"), *InFieldName));
+				PyUtil::SetPythonError(PyExc_Exception, PyType, *FString::Printf(TEXT("Method '%s' was set to 'override', but no method was found to override"), *InFieldName));
 				return false;
 			}
 			if (!SuperFunc->HasAnyFunctionFlags(FUNC_BlueprintEvent))
 			{
-				PyUtil::SetPythonError(PyExc_Exception, InPyType, *FString::Printf(TEXT("Method '%s' was set to 'override', but the method found to override was not a blueprint event"), *InFieldName));
+				PyUtil::SetPythonError(PyExc_Exception, PyType, *FString::Printf(TEXT("Method '%s' was set to 'override', but the method found to override was not a blueprint event"), *InFieldName));
 				return false;
 			}
 		}
@@ -1230,13 +1279,24 @@ struct FPythonGeneratedClassUtil
 		TArray<FPyObjectPtr> FuncArgDefaults;
 		if (!PyUtil::InspectFunctionArgs(InPyFuncDef->Func, FuncArgNames, &FuncArgDefaults))
 		{
-			PyUtil::SetPythonError(PyExc_Exception, InPyType, *FString::Printf(TEXT("Failed to inspect the arguments for '%s'"), *InFieldName));
+			PyUtil::SetPythonError(PyExc_Exception, PyType, *FString::Printf(TEXT("Failed to inspect the arguments for '%s'"), *InFieldName));
 			return false;
 		}
 
 		// Create the function, either from the definition, or from the super-function found to override
-		InClass->AddNativeFunction(*FuncName.ToString(), &UPythonGeneratedClass::CallPythonFunction); // Need to do this before the call to DuplicateObject in the case that the super-function already has FUNC_Native
-		UFunction* Func = SuperFunc ? DuplicateObject<UFunction>(SuperFunc, InClass, FuncName) : NewObject<UFunction>(InClass, FuncName);
+		NewClass->AddNativeFunction(*FuncName.ToString(), &UPythonGeneratedClass::CallPythonFunction); // Need to do this before the call to DuplicateObject in the case that the super-function already has FUNC_Native
+		UFunction* Func = nullptr;
+		if (SuperFunc)
+		{
+			FObjectDuplicationParameters FuncDuplicationParams(const_cast<UFunction*>(SuperFunc), NewClass);
+			FuncDuplicationParams.DestName = FuncName;
+			FuncDuplicationParams.InternalFlagMask &= ~EInternalObjectFlags::Native;
+			Func = CastChecked<UFunction>(StaticDuplicateObjectEx(FuncDuplicationParams));
+		}
+		else
+		{
+			Func = NewObject<UFunction>(NewClass, FuncName);
+		}
 		if (!SuperFunc)
 		{
 			Func->FunctionFlags |= FUNC_Public;
@@ -1263,7 +1323,7 @@ struct FPythonGeneratedClassUtil
 		}
 		Func->FunctionFlags |= (FUNC_Native | FUNC_Event | FUNC_BlueprintEvent | FUNC_BlueprintCallable);
 		FPyUFunctionDef::ApplyMetaData(InPyFuncDef, Func);
-		InClass->AddFunctionToFunctionMap(Func, Func->GetFName());
+		NewClass->AddFunctionToFunctionMap(Func, Func->GetFName());
 		if (!Func->HasAnyFunctionFlags(FUNC_Static))
 		{
 			// Strip the zero'th 'self' argument when processing a non-static function
@@ -1276,7 +1336,7 @@ struct FPythonGeneratedClassUtil
 			const int32 NumArgTypes = (InPyFuncDef->FuncParamTypes && InPyFuncDef->FuncParamTypes != Py_None) ? PySequence_Size(InPyFuncDef->FuncParamTypes) : 0;
 			if (NumArgTypes != FuncArgNames.Num())
 			{
-				PyUtil::SetPythonError(PyExc_Exception, InPyType, *FString::Printf(TEXT("Incorrect number of arguments specified for '%s' (expected %d, got %d)"), *InFieldName, NumArgTypes, FuncArgNames.Num()));
+				PyUtil::SetPythonError(PyExc_Exception, PyType, *FString::Printf(TEXT("Incorrect number of arguments specified for '%s' (expected %d, got %d)"), *InFieldName, NumArgTypes, FuncArgNames.Num()));
 				return false;
 			}
 
@@ -1290,7 +1350,7 @@ struct FPythonGeneratedClassUtil
 				UProperty* RetProp = PyUtil::CreateProperty(RetType, 1, Func, TEXT("ReturnValue"));
 				if (!RetProp)
 				{
-					PyUtil::SetPythonError(PyExc_Exception, InPyType, *FString::Printf(TEXT("Failed to create return property (%s) for function '%s'"), *PyUtil::GetFriendlyTypename(RetType), *InFieldName));
+					PyUtil::SetPythonError(PyExc_Exception, PyType, *FString::Printf(TEXT("Failed to create return property (%s) for function '%s'"), *PyUtil::GetFriendlyTypename(RetType), *InFieldName));
 					return false;
 				}
 				RetProp->PropertyFlags |= (CPF_Parm | CPF_ReturnParm);
@@ -1305,7 +1365,7 @@ struct FPythonGeneratedClassUtil
 						UProperty* ArgProp = PyUtil::CreateProperty(ArgTypeObj, 1, Func, *FString::Printf(TEXT("OutValue%d"), ArgIndex));
 						if (!ArgProp)
 						{
-							PyUtil::SetPythonError(PyExc_Exception, InPyType, *FString::Printf(TEXT("Failed to create output property (%s) for function '%s' at index %d"), *PyUtil::GetFriendlyTypename(ArgTypeObj), *InFieldName, ArgIndex));
+							PyUtil::SetPythonError(PyExc_Exception, PyType, *FString::Printf(TEXT("Failed to create output property (%s) for function '%s' at index %d"), *PyUtil::GetFriendlyTypename(ArgTypeObj), *InFieldName, ArgIndex));
 							return false;
 						}
 						ArgProp->PropertyFlags |= (CPF_Parm | CPF_OutParm);
@@ -1320,7 +1380,7 @@ struct FPythonGeneratedClassUtil
 				UProperty* ArgProp = PyUtil::CreateProperty(ArgTypeObj, 1, Func, *FuncArgNames[ArgIndex]);
 				if (!ArgProp)
 				{
-					PyUtil::SetPythonError(PyExc_Exception, InPyType, *FString::Printf(TEXT("Failed to create property (%s) for function '%s' argument '%s'"), *PyUtil::GetFriendlyTypename(ArgTypeObj), *InFieldName, *FuncArgNames[ArgIndex]));
+					PyUtil::SetPythonError(PyExc_Exception, PyType, *FString::Printf(TEXT("Failed to create property (%s) for function '%s' argument '%s'"), *PyUtil::GetFriendlyTypename(ArgTypeObj), *InFieldName, *FuncArgNames[ArgIndex]));
 					return false;
 				}
 				ArgProp->PropertyFlags |= CPF_Parm;
@@ -1343,9 +1403,9 @@ struct FPythonGeneratedClassUtil
 			{
 				// Convert the default value to the given property...
 				PyUtil::FPropValueOnScope DefaultValue(Param);
-				if (!DefaultValue.IsValid() || !DefaultValue.SetValue(FuncArgDefaults[InputArgIndex], *PyUtil::GetErrorContext(InPyType)))
+				if (!DefaultValue.IsValid() || !DefaultValue.SetValue(FuncArgDefaults[InputArgIndex], *PyUtil::GetErrorContext(PyType)))
 				{
-					PyUtil::SetPythonError(PyExc_Exception, InPyType, *FString::Printf(TEXT("Failed to convert default value for function '%s' argument '%s' (%s)"), *InFieldName, *FuncArgNames[InputArgIndex], *Param->GetClass()->GetName()));
+					PyUtil::SetPythonError(PyExc_Exception, PyType, *FString::Printf(TEXT("Failed to convert default value for function '%s' argument '%s' (%s)"), *InFieldName, *FuncArgNames[InputArgIndex], *Param->GetClass()->GetName()));
 					return false;
 				}
 
@@ -1353,7 +1413,7 @@ struct FPythonGeneratedClassUtil
 				FString ExportedDefaultValue;
 				if (!DefaultValue.GetProp()->ExportText_Direct(ExportedDefaultValue, DefaultValue.GetValue(), DefaultValue.GetValue(), nullptr, PPF_None))
 				{
-					PyUtil::SetPythonError(PyExc_Exception, InPyType, *FString::Printf(TEXT("Failed to export default value for function '%s' argument '%s' (%s)"), *InFieldName, *FuncArgNames[InputArgIndex], *Param->GetClass()->GetName()));
+					PyUtil::SetPythonError(PyExc_Exception, PyType, *FString::Printf(TEXT("Failed to export default value for function '%s' argument '%s' (%s)"), *InFieldName, *FuncArgNames[InputArgIndex], *Param->GetClass()->GetName()));
 					return false;
 				}
 
@@ -1380,7 +1440,7 @@ struct FPythonGeneratedClassUtil
 
 		if (GeneratedWrappedFunction.InputParams.Num() != FuncArgNames.Num())
 		{
-			PyUtil::SetPythonError(PyExc_Exception, InPyType, *FString::Printf(TEXT("Incorrect number of arguments specified for '%s' (expected %d, got %d)"), *InFieldName, GeneratedWrappedFunction.InputParams.Num(), FuncArgNames.Num()));
+			PyUtil::SetPythonError(PyExc_Exception, PyType, *FString::Printf(TEXT("Incorrect number of arguments specified for '%s' (expected %d, got %d)"), *InFieldName, GeneratedWrappedFunction.InputParams.Num(), FuncArgNames.Num()));
 			return false;
 		}
 
@@ -1396,7 +1456,7 @@ struct FPythonGeneratedClassUtil
 		}
 
 		// Build the definition data for the new method
-		PyGenUtil::FFunctionDef& FuncDef = *InClass->FunctionDefs.Add_GetRef(MakeShared<PyGenUtil::FFunctionDef>());
+		PyGenUtil::FFunctionDef& FuncDef = *NewClass->FunctionDefs.Add_GetRef(MakeShared<PyGenUtil::FFunctionDef>());
 		FuncDef.GeneratedWrappedMethod.MethodName = PyGenUtil::TCHARToUTF8Buffer(*InFieldName);
 		FuncDef.GeneratedWrappedMethod.MethodDoc = PyGenUtil::TCHARToUTF8Buffer(*PyGenUtil::GetFieldTooltip(Func));
 		FuncDef.GeneratedWrappedMethod.MethodFunc = MoveTemp(GeneratedWrappedFunction);
@@ -1417,35 +1477,37 @@ struct FPythonGeneratedClassUtil
 		return true;
 	}
 
-	static bool CopyPropertiesFromOldClass(UPythonGeneratedClass* InClass, UPythonGeneratedClass* InOldClass, PyTypeObject* InPyType)
+	bool CopyPropertiesFromOldClass()
 	{
-		InClass->PropertyDefs.Reserve(InOldClass->PropertyDefs.Num());
-		for (const TSharedPtr<PyGenUtil::FPropertyDef>& OldPropDef : InOldClass->PropertyDefs)
+		check(OldClass);
+
+		NewClass->PropertyDefs.Reserve(OldClass->PropertyDefs.Num());
+		for (const TSharedPtr<PyGenUtil::FPropertyDef>& OldPropDef : OldClass->PropertyDefs)
 		{
 			const UProperty* OldProp = OldPropDef->GeneratedWrappedGetSet.Prop.Prop;
 			const UFunction* OldGetter = OldPropDef->GeneratedWrappedGetSet.GetFunc.Func;
 			const UFunction* OldSetter = OldPropDef->GeneratedWrappedGetSet.SetFunc.Func;
 
-			UProperty* Prop = DuplicateObject<UProperty>(OldProp, InClass, OldProp->GetFName());
+			UProperty* Prop = DuplicateObject<UProperty>(OldProp, NewClass, OldProp->GetFName());
 			if (!Prop)
 			{
-				PyUtil::SetPythonError(PyExc_Exception, InPyType, *FString::Printf(TEXT("Failed to duplicate property for '%s'"), UTF8_TO_TCHAR(OldPropDef->PyGetSet.name)));
+				PyUtil::SetPythonError(PyExc_Exception, PyType, *FString::Printf(TEXT("Failed to duplicate property for '%s'"), UTF8_TO_TCHAR(OldPropDef->PyGetSet.name)));
 				return false;
 			}
 
 			UMetaData::CopyMetadata((UProperty*)OldProp, Prop);
-			InClass->AddCppProperty(Prop);
+			NewClass->AddCppProperty(Prop);
 
-			PyGenUtil::FPropertyDef& PropDef = *InClass->PropertyDefs.Add_GetRef(MakeShared<PyGenUtil::FPropertyDef>());
+			PyGenUtil::FPropertyDef& PropDef = *NewClass->PropertyDefs.Add_GetRef(MakeShared<PyGenUtil::FPropertyDef>());
 			PropDef.GeneratedWrappedGetSet = OldPropDef->GeneratedWrappedGetSet;
 			PropDef.GeneratedWrappedGetSet.Prop.SetProperty(Prop);
 			if (OldGetter)
 			{
-				PropDef.GeneratedWrappedGetSet.GetFunc.SetFunction(InClass->FindFunctionByName(OldGetter->GetFName()));
+				PropDef.GeneratedWrappedGetSet.GetFunc.SetFunction(NewClass->FindFunctionByName(OldGetter->GetFName()));
 			}
 			if (OldSetter)
 			{
-				PropDef.GeneratedWrappedGetSet.SetFunc.SetFunction(InClass->FindFunctionByName(OldSetter->GetFName()));
+				PropDef.GeneratedWrappedGetSet.SetFunc.SetFunction(NewClass->FindFunctionByName(OldSetter->GetFName()));
 			}
 			PropDef.GeneratedWrappedGetSet.ToPython(PropDef.PyGetSet);
 		}
@@ -1453,28 +1515,30 @@ struct FPythonGeneratedClassUtil
 		return true;
 	}
 
-	static bool CopyFunctionsFromOldClass(UPythonGeneratedClass* InClass, UPythonGeneratedClass* InOldClass, PyTypeObject* InPyType)
+	bool CopyFunctionsFromOldClass()
 	{
-		InClass->FunctionDefs.Reserve(InOldClass->FunctionDefs.Num());
-		for (const TSharedPtr<PyGenUtil::FFunctionDef>& OldFuncDef : InOldClass->FunctionDefs)
+		check(OldClass);
+
+		NewClass->FunctionDefs.Reserve(OldClass->FunctionDefs.Num());
+		for (const TSharedPtr<PyGenUtil::FFunctionDef>& OldFuncDef : OldClass->FunctionDefs)
 		{
 			const UFunction* OldFunc = OldFuncDef->GeneratedWrappedMethod.MethodFunc.Func;
 
-			InClass->AddNativeFunction(*OldFunc->GetName(), &UPythonGeneratedClass::CallPythonFunction);
-			UFunction* Func = DuplicateObject<UFunction>(OldFunc, InClass, OldFunc->GetFName());
+			NewClass->AddNativeFunction(*OldFunc->GetName(), &UPythonGeneratedClass::CallPythonFunction);
+			UFunction* Func = DuplicateObject<UFunction>(OldFunc, NewClass, OldFunc->GetFName());
 			if (!Func)
 			{
-				PyUtil::SetPythonError(PyExc_Exception, InPyType, *FString::Printf(TEXT("Failed to duplicate function for '%s'"), UTF8_TO_TCHAR(OldFuncDef->PyMethod.MethodName)));
+				PyUtil::SetPythonError(PyExc_Exception, PyType, *FString::Printf(TEXT("Failed to duplicate function for '%s'"), UTF8_TO_TCHAR(OldFuncDef->PyMethod.MethodName)));
 				return false;
 			}
 
 			UMetaData::CopyMetadata((UFunction*)OldFunc, Func);
-			InClass->AddFunctionToFunctionMap(Func, Func->GetFName());
+			NewClass->AddFunctionToFunctionMap(Func, Func->GetFName());
 
 			Func->Bind();
 			Func->StaticLink(true);
 
-			PyGenUtil::FFunctionDef& FuncDef = *InClass->FunctionDefs.Add_GetRef(MakeShared<PyGenUtil::FFunctionDef>());
+			PyGenUtil::FFunctionDef& FuncDef = *NewClass->FunctionDefs.Add_GetRef(MakeShared<PyGenUtil::FFunctionDef>());
 			FuncDef.GeneratedWrappedMethod = OldFuncDef->GeneratedWrappedMethod;
 			FuncDef.GeneratedWrappedMethod.MethodFunc.SetFunction(Func);
 			FuncDef.PyFunction = OldFuncDef->PyFunction;
@@ -1485,40 +1549,64 @@ struct FPythonGeneratedClassUtil
 		return true;
 	}
 
-	static bool RegisterDescriptors(UPythonGeneratedClass* InClass, PyTypeObject* InPyType)
+	void ReparentPythonType(PyTypeObject* InNewBasePyType)
+	{
+		auto UpdateTuple = [](PyObject* InTuple, PyTypeObject* InOldType, PyTypeObject* InNewType)
+		{
+			if (InTuple)
+			{
+				const int32 TupleSize = PyTuple_Size(InTuple);
+				for (int32 TupleIndex = 0; TupleIndex < TupleSize; ++TupleIndex)
+				{
+					if (PyTuple_GetItem(InTuple, TupleIndex) == (PyObject*)InOldType)
+					{
+						FPyTypeObjectPtr NewType = FPyTypeObjectPtr::NewReference(InNewType);
+						PyTuple_SetItem(InTuple, TupleIndex, (PyObject*)NewType.Release()); // PyTuple_SetItem steals the reference
+					}
+				}
+			}
+		};
+
+		UpdateTuple(PyType->tp_bases, PyType->tp_base, InNewBasePyType);
+		UpdateTuple(PyType->tp_mro, PyType->tp_base, InNewBasePyType);
+		PyType->tp_base = InNewBasePyType;
+	}
+
+private:
+	bool RegisterDescriptors()
 	{
-		for (const TSharedPtr<PyGenUtil::FPropertyDef>& PropDef : InClass->PropertyDefs)
+		for (const TSharedPtr<PyGenUtil::FPropertyDef>& PropDef : NewClass->PropertyDefs)
 		{
-			FPyObjectPtr GetSetDesc = FPyObjectPtr::StealReference(PyDescr_NewGetSet(InPyType, &PropDef->PyGetSet));
+			FPyObjectPtr GetSetDesc = FPyObjectPtr::StealReference(PyDescr_NewGetSet(PyType, &PropDef->PyGetSet));
 			if (!GetSetDesc)
 			{
-				PyUtil::SetPythonError(PyExc_Exception, InPyType, *FString::Printf(TEXT("Failed to create descriptor for '%s'"), UTF8_TO_TCHAR(PropDef->PyGetSet.name)));
+				PyUtil::SetPythonError(PyExc_Exception, PyType, *FString::Printf(TEXT("Failed to create descriptor for '%s'"), UTF8_TO_TCHAR(PropDef->PyGetSet.name)));
 				return false;
 			}
-			if (PyDict_SetItemString(InPyType->tp_dict, PropDef->PyGetSet.name, GetSetDesc) != 0)
+			if (PyDict_SetItemString(PyType->tp_dict, PropDef->PyGetSet.name, GetSetDesc) != 0)
 			{
-				PyUtil::SetPythonError(PyExc_Exception, InPyType, *FString::Printf(TEXT("Failed to assign descriptor for '%s'"), UTF8_TO_TCHAR(PropDef->PyGetSet.name)));
+				PyUtil::SetPythonError(PyExc_Exception, PyType, *FString::Printf(TEXT("Failed to assign descriptor for '%s'"), UTF8_TO_TCHAR(PropDef->PyGetSet.name)));
 				return false;
 			}
 		}
 
-		for (const TSharedPtr<PyGenUtil::FFunctionDef>& FuncDef : InClass->FunctionDefs)
+		for (const TSharedPtr<PyGenUtil::FFunctionDef>& FuncDef : NewClass->FunctionDefs)
 		{
 			if (FuncDef->bIsHidden)
 			{
-				PyDict_DelItemString(InPyType->tp_dict, FuncDef->PyMethod.MethodName);
+				PyDict_DelItemString(PyType->tp_dict, FuncDef->PyMethod.MethodName);
 			}
 			else
 			{
-				FPyObjectPtr MethodDesc = FPyObjectPtr::StealReference(FPyMethodWithClosureDef::NewMethodDescriptor(InPyType, &FuncDef->PyMethod));
+				FPyObjectPtr MethodDesc = FPyObjectPtr::StealReference(FPyMethodWithClosureDef::NewMethodDescriptor(PyType, &FuncDef->PyMethod));
 				if (!MethodDesc)
 				{
-					PyUtil::SetPythonError(PyExc_Exception, InPyType, *FString::Printf(TEXT("Failed to create descriptor for '%s'"), UTF8_TO_TCHAR(FuncDef->PyMethod.MethodName)));
+					PyUtil::SetPythonError(PyExc_Exception, PyType, *FString::Printf(TEXT("Failed to create descriptor for '%s'"), UTF8_TO_TCHAR(FuncDef->PyMethod.MethodName)));
 					return false;
 				}
-				if (PyDict_SetItemString(InPyType->tp_dict, FuncDef->PyMethod.MethodName, MethodDesc) != 0)
+				if (PyDict_SetItemString(PyType->tp_dict, FuncDef->PyMethod.MethodName, MethodDesc) != 0)
 				{
-					PyUtil::SetPythonError(PyExc_Exception, InPyType, *FString::Printf(TEXT("Failed to assign descriptor for '%s'"), UTF8_TO_TCHAR(FuncDef->PyMethod.MethodName)));
+					PyUtil::SetPythonError(PyExc_Exception, PyType, *FString::Printf(TEXT("Failed to assign descriptor for '%s'"), UTF8_TO_TCHAR(FuncDef->PyMethod.MethodName)));
 					return false;
 				}
 			}
@@ -1527,36 +1615,32 @@ struct FPythonGeneratedClassUtil
 		return true;
 	}
 
-	static void ReparentPythonType(PyTypeObject* InPyType, PyTypeObject* InNewBasePyType)
+	void PrepareOldClassForReinstancing()
 	{
-		auto UpdateTuple = [](PyObject* InTuple, PyTypeObject* InOldType, PyTypeObject* InNewType)
-		{
-			if (InTuple)
-			{
-				const int32 TupleSize = PyTuple_Size(InTuple);
-				for (int32 TupleIndex = 0; TupleIndex < TupleSize; ++TupleIndex)
-				{
-					if (PyTuple_GetItem(InTuple, TupleIndex) == (PyObject*)InOldType)
-					{
-						FPyTypeObjectPtr NewType = FPyTypeObjectPtr::NewReference(InNewType);
-						PyTuple_SetItem(InTuple, TupleIndex, (PyObject*)NewType.Release()); // PyTuple_SetItem steals the reference
-					}
-				}
-			}
-		};
+		check(OldClass);
 
-		UpdateTuple(InPyType->tp_bases, InPyType->tp_base, InNewBasePyType);
-		UpdateTuple(InPyType->tp_mro, InPyType->tp_base, InNewBasePyType);
-		InPyType->tp_base = InNewBasePyType;
+		const FString OldClassName = MakeUniqueObjectName(OldClass->GetOuter(), OldClass->GetClass(), *FString::Printf(TEXT("%s_REINST"), *OldClass->GetName())).ToString();
+		OldClass->ClassFlags |= CLASS_NewerVersionExists;
+		OldClass->SetFlags(RF_NewerVersionExists);
+		OldClass->ClearFlags(RF_Public | RF_Standalone);
+		OldClass->Rename(*OldClassName, nullptr, REN_DontCreateRedirectors);
 	}
+
+	FString ClassName;
+	PyTypeObject* PyType;
+	UPythonGeneratedClass* OldClass;
+	UPythonGeneratedClass* NewClass;
 };
 
 void UPythonGeneratedClass::PostRename(UObject* OldOuter, const FName OldName)
 {
 	Super::PostRename(OldOuter, OldName);
 
-	FPyWrapperTypeRegistry::Get().UnregisterWrappedClassType(OldName, PyType);
-	FPyWrapperTypeRegistry::Get().RegisterWrappedClassType(GetFName(), PyType, !HasAnyFlags(RF_NewerVersionExists));
+	if (PyType)
+	{
+		FPyWrapperTypeRegistry::Get().UnregisterWrappedClassType(OldName, PyType);
+		FPyWrapperTypeRegistry::Get().RegisterWrappedClassType(GetFName(), PyType, !HasAnyFlags(RF_NewerVersionExists));
+	}
 }
 
 void UPythonGeneratedClass::PostInitInstance(UObject* InObj)
@@ -1587,9 +1671,6 @@ void UPythonGeneratedClass::PostInitInstance(UObject* InObj)
 
 UPythonGeneratedClass* UPythonGeneratedClass::GenerateClass(PyTypeObject* InPyType)
 {
-	UObject* ClassOuter = GetPythonTypeContainer();
-	const FString ClassName = PyUtil::GetCleanTypename(InPyType);
-
 	// Get the correct super class from the parent type in Python
 	UClass* SuperClass = FPyWrapperObjectMetaData::GetClass(InPyType->tp_base);
 	if (!SuperClass)
@@ -1598,20 +1679,8 @@ UPythonGeneratedClass* UPythonGeneratedClass::GenerateClass(PyTypeObject* InPyTy
 		return nullptr;
 	}
 
-	UPythonGeneratedClass* OldClass = FindObject<UPythonGeneratedClass>(ClassOuter, *ClassName);
-	if (OldClass)
-	{
-		FPythonGeneratedClassUtil::PrepareOldClassForReinstancing(OldClass);
-	}
-
-	UPythonGeneratedClass* Class = FPythonGeneratedClassUtil::CreateClass(ClassName, ClassOuter, SuperClass);
-
-	// Get the post-init function
-	Class->PyPostInitFunction = FPyObjectPtr::StealReference(PyGenUtil::GetPostInitFunc(InPyType));
-	if (!Class->PyPostInitFunction)
-	{
-		return nullptr;
-	}
+	// Builder used to generate the class
+	FPythonGeneratedClassBuilder PythonClassBuilder(PyUtil::GetCleanTypename(InPyType), SuperClass, InPyType);
 
 	// Add the functions to this class
 	// We have to process these first as properties may reference them as get/set functions
@@ -1633,7 +1702,7 @@ UPythonGeneratedClass* UPythonGeneratedClass::GenerateClass(PyTypeObject* InPyTy
 			if (PyObject_IsInstance(FieldValue, (PyObject*)&PyUFunctionDefType) == 1)
 			{
 				FPyUFunctionDef* PyFuncDef = (FPyUFunctionDef*)FieldValue;
-				if (!FPythonGeneratedClassUtil::CreateFunctionFromDefinition(Class, InPyType, FieldName, PyFuncDef))
+				if (!PythonClassBuilder.CreateFunctionFromDefinition(FieldName, PyFuncDef))
 				{
 					return nullptr;
 				}
@@ -1653,7 +1722,7 @@ UPythonGeneratedClass* UPythonGeneratedClass::GenerateClass(PyTypeObject* InPyTy
 			if (PyObject_IsInstance(FieldValue, (PyObject*)&PyUPropertyDefType) == 1)
 			{
 				FPyUPropertyDef* PyPropDef = (FPyUPropertyDef*)FieldValue;
-				if (!FPythonGeneratedClassUtil::CreatePropertyFromDefinition(Class, InPyType, FieldName, PyPropDef))
+				if (!PythonClassBuilder.CreatePropertyFromDefinition(FieldName, PyPropDef))
 				{
 					return nullptr;
 				}
@@ -1661,26 +1730,8 @@ UPythonGeneratedClass* UPythonGeneratedClass::GenerateClass(PyTypeObject* InPyTy
 		}
 	}
 
-	// Replace the definitions with real descriptors
-	if (!FPythonGeneratedClassUtil::RegisterDescriptors(Class, InPyType))
-	{
-		return nullptr;
-	}
-
-	// Let Python know that we've changed its type
-	PyType_Modified(InPyType);
-
-	// Finalize the class
-	FPythonGeneratedClassUtil::FinalizeClass(Class, InPyType);
-
-	// Re-instance the old class and re-parent any derived classes to this new type
-	if (OldClass)
-	{
-		FPyWrapperTypeReinstancer::Get().AddPendingClass(OldClass, Class);
-		ReparentDerivedClasses(OldClass, Class);
-	}
-
-	return Class;
+	// Finalize the class with its post-init function
+	return PythonClassBuilder.Finalize(FPyObjectPtr::StealReference(PyGenUtil::GetPostInitFunc(InPyType)));
 }
 
 bool UPythonGeneratedClass::ReparentDerivedClasses(UPythonGeneratedClass* InOldParent, UPythonGeneratedClass* InNewParent)
@@ -1710,44 +1761,26 @@ bool UPythonGeneratedClass::ReparentDerivedClasses(UPythonGeneratedClass* InOldP
 
 UPythonGeneratedClass* UPythonGeneratedClass::ReparentClass(UPythonGeneratedClass* InOldClass, UPythonGeneratedClass* InNewParent)
 {
-	UObject* ClassOuter = GetPythonTypeContainer();
-	const FString ClassName = InOldClass->GetName();
-
-	FPythonGeneratedClassUtil::PrepareOldClassForReinstancing(InOldClass);
-	UPythonGeneratedClass* Class = FPythonGeneratedClassUtil::CreateClass(ClassName, ClassOuter, InNewParent);
-	PyTypeObject* PyType = InOldClass->PyType;
+	// Builder used to generate the class
+	FPythonGeneratedClassBuilder PythonClassBuilder(InOldClass->GetName(), InNewParent, InOldClass->PyType);
 
 	// Copy the data from the old class
-	Class->PyPostInitFunction = InOldClass->PyPostInitFunction;
-	if (!FPythonGeneratedClassUtil::CopyFunctionsFromOldClass(Class, InOldClass, PyType))
+	if (!PythonClassBuilder.CopyFunctionsFromOldClass())
 	{
 		return nullptr;
 	}
-	if (!FPythonGeneratedClassUtil::CopyPropertiesFromOldClass(Class, InOldClass, PyType))
+	if (!PythonClassBuilder.CopyPropertiesFromOldClass())
 	{
 		return nullptr;
 	}
 
-	// Update the descriptors on the type so they reference the new class
-	if (!FPythonGeneratedClassUtil::RegisterDescriptors(Class, PyType))
+	UPythonGeneratedClass* NewClass = PythonClassBuilder.Finalize(InOldClass->PyPostInitFunction);
+	if (NewClass)
 	{
-		return nullptr;
+		// Update the base of the Python type
+		PythonClassBuilder.ReparentPythonType(InNewParent->PyType);
 	}
-
-	// Update the base of the Python type
-	FPythonGeneratedClassUtil::ReparentPythonType(PyType, InNewParent->PyType);
-
-	// Let Python know that we've changed its type
-	PyType_Modified(PyType);
-
-	// Finalize the class
-	FPythonGeneratedClassUtil::FinalizeClass(Class, PyType);
-
-	// Re-instance the old class and re-parent any derived classes to this new type
-	FPyWrapperTypeReinstancer::Get().AddPendingClass(InOldClass, Class);
-	ReparentDerivedClasses(InOldClass, Class);
-
-	return Class;
+	return NewClass;
 }
 
 DEFINE_FUNCTION(UPythonGeneratedClass::CallPythonFunction)
diff --git a/Engine/Plugins/Experimental/PythonScriptPlugin/Source/PythonScriptPlugin/Private/PyWrapperObject.h b/Engine/Plugins/Experimental/PythonScriptPlugin/Source/PythonScriptPlugin/Private/PyWrapperObject.h
--- a/Engine/Plugins/Experimental/PythonScriptPlugin/Source/PythonScriptPlugin/Private/PyWrapperObject.h
+++ b/Engine/Plugins/Experimental/PythonScriptPlugin/Source/PythonScriptPlugin/Private/PyWrapperObject.h
@@ -221,7 +221,7 @@ private:
 	/** Meta-data for this generated class that is applied to the Python type */
 	FPyWrapperObjectMetaData PyMetaData;
 
-	friend struct FPythonGeneratedClassUtil;
+	friend class FPythonGeneratedClassBuilder;
 
 #endif	// WITH_PYTHON
 };
diff --git a/Engine/Plugins/Experimental/PythonScriptPlugin/Source/PythonScriptPlugin/Private/PyWrapperStruct.cpp b/Engine/Plugins/Experimental/PythonScriptPlugin/Source/PythonScriptPlugin/Private/PyWrapperStruct.cpp
--- a/Engine/Plugins/Experimental/PythonScriptPlugin/Source/PythonScriptPlugin/Private/PyWrapperStruct.cpp
+++ b/Engine/Plugins/Experimental/PythonScriptPlugin/Source/PythonScriptPlugin/Private/PyWrapperStruct.cpp
@@ -1279,33 +1279,64 @@ bool FPyWrapperStructMetaData::IsStructDeprecated(FPyWrapperStruct* Instance, FS
 	return IsStructDeprecated(Py_TYPE(Instance), OutDeprecationMessage);
 }
 
-struct FPythonGeneratedStructUtil
+class FPythonGeneratedStructBuilder
 {
-	static void PrepareOldStructForReinstancing(UPythonGeneratedStruct* InOldStruct)
+public:
+	FPythonGeneratedStructBuilder(const FString& InStructName, UScriptStruct* InSuperStruct, PyTypeObject* InPyType)
+		: StructName(InStructName)
+		, PyType(InPyType)
+		, OldStruct(nullptr)
+		, NewStruct(nullptr)
 	{
-		const FString OldStructName = MakeUniqueObjectName(InOldStruct->GetOuter(), InOldStruct->GetClass(), *FString::Printf(TEXT("%s_REINST"), *InOldStruct->GetName())).ToString();
-		InOldStruct->SetFlags(RF_NewerVersionExists);
-		InOldStruct->ClearFlags(RF_Public | RF_Standalone);
-		InOldStruct->Rename(*OldStructName, nullptr, REN_DontCreateRedirectors);
+		UObject* StructOuter = GetPythonTypeContainer();
+
+		// Find any existing struct with the name we want to use
+		OldStruct = FindObject<UPythonGeneratedStruct>(StructOuter, *StructName);
+
+		// Create a new struct with a temporary name; we will rename it as part of Finalize
+		const FString NewStructName = MakeUniqueObjectName(StructOuter, UPythonGeneratedStruct::StaticClass(), *FString::Printf(TEXT("%s_NEWINST"), *StructName)).ToString();
+		NewStruct = NewObject<UPythonGeneratedStruct>(StructOuter, *NewStructName, RF_Public | RF_Standalone | RF_Transient);
+		NewStruct->SetMetaData(TEXT("BlueprintType"), TEXT("true"));
+		NewStruct->SetSuperStruct(InSuperStruct);
 	}
 
-	static UPythonGeneratedStruct* CreateStruct(const FString& InStructName, UObject* InStructOuter, UScriptStruct* InSuperStruct)
+	~FPythonGeneratedStructBuilder()
 	{
-		UPythonGeneratedStruct* Struct = NewObject<UPythonGeneratedStruct>(InStructOuter, *InStructName, RF_Public | RF_Standalone);
-		Struct->SetMetaData(TEXT("BlueprintType"), TEXT("true"));
-		Struct->SetSuperStruct(InSuperStruct);
-		return Struct;
+		// If NewStruct is still set at this point, if means Finalize wasn't called and we should destroy the partially built struct
+		if (NewStruct)
+		{
+			NewStruct->ClearFlags(RF_Public | RF_Standalone);
+			NewStruct = nullptr;
+
+			CollectGarbage(GARBAGE_COLLECTION_KEEPFLAGS);
+		}
 	}
 
-	static void FinalizeStruct(UPythonGeneratedStruct* InStruct, PyTypeObject* InPyType)
+	UPythonGeneratedStruct* Finalize(FPyObjectPtr InPyPostInitFunction)
 	{
+		// Set the post-init function
+		NewStruct->PyPostInitFunction = InPyPostInitFunction;
+		if (!NewStruct->PyPostInitFunction)
+		{
+			return nullptr;
+		}
+
+		// Replace the definitions with real descriptors
+		if (!RegisterDescriptors())
+		{
+			return nullptr;
+		}
+
+		// Let Python know that we've changed its type
+		PyType_Modified(PyType);
+
 		// Build a complete list of init params for this struct
 		TArray<PyGenUtil::FGeneratedWrappedMethodParameter> StructInitParams;
-		if (const FPyWrapperStructMetaData* SuperMetaData = FPyWrapperStructMetaData::GetMetaData(InPyType->tp_base))
+		if (const FPyWrapperStructMetaData* SuperMetaData = FPyWrapperStructMetaData::GetMetaData(PyType->tp_base))
 		{
 			StructInitParams = SuperMetaData->InitParams;
 		}
-		for (const TSharedPtr<PyGenUtil::FPropertyDef>& PropDef : InStruct->PropertyDefs)
+		for (const TSharedPtr<PyGenUtil::FPropertyDef>& PropDef : NewStruct->PropertyDefs)
 		{
 			if (!PropDef->GeneratedWrappedGetSet.Prop.DeprecationMessage.IsSet())
 			{
@@ -1316,52 +1347,66 @@ struct FPythonGeneratedStructUtil
 			}
 		}
 
+		// We can no longer fail, set the correct name on the new struct
+		NewStruct->Rename(*StructName, nullptr, REN_DontCreateRedirectors);
+
 		// Finalize the struct
-		InStruct->Bind();
-		InStruct->StaticLink(true);
+		NewStruct->Bind();
+		NewStruct->StaticLink(true);
 
 		// Add the object meta-data to the type
-		InStruct->PyMetaData.Struct = InStruct;
-		InStruct->PyMetaData.InitParams = MoveTemp(StructInitParams);
-		FPyWrapperStructMetaData::SetMetaData(InPyType, &InStruct->PyMetaData);
+		NewStruct->PyMetaData.Struct = NewStruct;
+		NewStruct->PyMetaData.InitParams = MoveTemp(StructInitParams);
+		FPyWrapperStructMetaData::SetMetaData(PyType, &NewStruct->PyMetaData);
 
 		// Map the Unreal struct to the Python type
-		InStruct->PyType = FPyTypeObjectPtr::NewReference(InPyType);
-		FPyWrapperTypeRegistry::Get().RegisterWrappedStructType(InStruct->GetFName(), InPyType);
+		NewStruct->PyType = FPyTypeObjectPtr::NewReference(PyType);
+		FPyWrapperTypeRegistry::Get().RegisterWrappedStructType(NewStruct->GetFName(), PyType);
+
+		// Re-instance the old struct
+		if (OldStruct)
+		{
+			FPyWrapperTypeReinstancer::Get().AddPendingStruct(OldStruct, NewStruct);
+		}
+
+		// Null the NewStruct pointer so the destructor doesn't kill it
+		UPythonGeneratedStruct* FinalizedStruct = NewStruct;
+		NewStruct = nullptr;
+		return FinalizedStruct;
 	}
 
-	static bool CreatePropertyFromDefinition(UPythonGeneratedStruct* InStruct, PyTypeObject* InPyType, const FString& InFieldName, FPyUPropertyDef* InPyPropDef)
+	bool CreatePropertyFromDefinition(const FString& InFieldName, FPyUPropertyDef* InPyPropDef)
 	{
-		UScriptStruct* SuperStruct = Cast<UScriptStruct>(InStruct->GetSuperStruct());
+		UScriptStruct* SuperStruct = Cast<UScriptStruct>(NewStruct->GetSuperStruct());
 
 		// Resolve the property name to match any previously exported properties from the parent type
-		const FName PropName = FPyWrapperStructMetaData::ResolvePropertyName(InPyType->tp_base, *InFieldName);
+		const FName PropName = FPyWrapperStructMetaData::ResolvePropertyName(PyType->tp_base, *InFieldName);
 		if (SuperStruct && SuperStruct->FindPropertyByName(PropName))
 		{
-			PyUtil::SetPythonError(PyExc_Exception, InPyType, *FString::Printf(TEXT("Property '%s' (%s) cannot override a property from the base type"), *InFieldName, *PyUtil::GetFriendlyTypename(InPyPropDef->PropType)));
+			PyUtil::SetPythonError(PyExc_Exception, PyType, *FString::Printf(TEXT("Property '%s' (%s) cannot override a property from the base type"), *InFieldName, *PyUtil::GetFriendlyTypename(InPyPropDef->PropType)));
 			return false;
 		}
 
 		// Structs cannot support getter/setter functions (or any functions)
 		if (!InPyPropDef->GetterFuncName.IsEmpty() || !InPyPropDef->SetterFuncName.IsEmpty())
 		{
-			PyUtil::SetPythonError(PyExc_Exception, InPyType, *FString::Printf(TEXT("Struct property '%s' (%s) has a getter or setter, which is not supported on structs"), *InFieldName, *PyUtil::GetFriendlyTypename(InPyPropDef->PropType)));
+			PyUtil::SetPythonError(PyExc_Exception, PyType, *FString::Printf(TEXT("Struct property '%s' (%s) has a getter or setter, which is not supported on structs"), *InFieldName, *PyUtil::GetFriendlyTypename(InPyPropDef->PropType)));
 			return false;
 		}
 
 		// Create the property from its definition
-		UProperty* Prop = PyUtil::CreateProperty(InPyPropDef->PropType, 1, InStruct, PropName);
+		UProperty* Prop = PyUtil::CreateProperty(InPyPropDef->PropType, 1, NewStruct, PropName);
 		if (!Prop)
 		{
-			PyUtil::SetPythonError(PyExc_Exception, InPyType, *FString::Printf(TEXT("Failed to create property for '%s' (%s)"), *InFieldName, *PyUtil::GetFriendlyTypename(InPyPropDef->PropType)));
+			PyUtil::SetPythonError(PyExc_Exception, PyType, *FString::Printf(TEXT("Failed to create property for '%s' (%s)"), *InFieldName, *PyUtil::GetFriendlyTypename(InPyPropDef->PropType)));
 			return false;
 		}
 		Prop->PropertyFlags |= (CPF_Edit | CPF_BlueprintVisible);
 		FPyUPropertyDef::ApplyMetaData(InPyPropDef, Prop);
-		InStruct->AddCppProperty(Prop);
+		NewStruct->AddCppProperty(Prop);
 
 		// Build the definition data for the new property accessor
-		PyGenUtil::FPropertyDef& PropDef = *InStruct->PropertyDefs.Add_GetRef(MakeShared<PyGenUtil::FPropertyDef>());
+		PyGenUtil::FPropertyDef& PropDef = *NewStruct->PropertyDefs.Add_GetRef(MakeShared<PyGenUtil::FPropertyDef>());
 		PropDef.GeneratedWrappedGetSet.GetSetName = PyGenUtil::TCHARToUTF8Buffer(*InFieldName);
 		PropDef.GeneratedWrappedGetSet.GetSetDoc = PyGenUtil::TCHARToUTF8Buffer(*FString::Printf(TEXT("type: %s\n%s"), *PyGenUtil::GetPropertyPythonType(Prop), *PyGenUtil::GetFieldTooltip(Prop)));
 		PropDef.GeneratedWrappedGetSet.Prop.SetProperty(Prop);
@@ -1372,33 +1417,52 @@ struct FPythonGeneratedStructUtil
 		return true;
 	}
 
-	static bool RegisterDescriptors(UPythonGeneratedStruct* InStruct, PyTypeObject* InPyType)
+private:
+	bool RegisterDescriptors()
 	{
-		for (const TSharedPtr<PyGenUtil::FPropertyDef>& PropDef : InStruct->PropertyDefs)
+		for (const TSharedPtr<PyGenUtil::FPropertyDef>& PropDef : NewStruct->PropertyDefs)
 		{
-			FPyObjectPtr GetSetDesc = FPyObjectPtr::StealReference(PyDescr_NewGetSet(InPyType, &PropDef->PyGetSet));
+			FPyObjectPtr GetSetDesc = FPyObjectPtr::StealReference(PyDescr_NewGetSet(PyType, &PropDef->PyGetSet));
 			if (!GetSetDesc)
 			{
-				PyUtil::SetPythonError(PyExc_Exception, InPyType, *FString::Printf(TEXT("Failed to create descriptor for '%s'"), UTF8_TO_TCHAR(PropDef->PyGetSet.name)));
+				PyUtil::SetPythonError(PyExc_Exception, PyType, *FString::Printf(TEXT("Failed to create descriptor for '%s'"), UTF8_TO_TCHAR(PropDef->PyGetSet.name)));
 				return false;
 			}
-			if (PyDict_SetItemString(InPyType->tp_dict, PropDef->PyGetSet.name, GetSetDesc) != 0)
+			if (PyDict_SetItemString(PyType->tp_dict, PropDef->PyGetSet.name, GetSetDesc) != 0)
 			{
-				PyUtil::SetPythonError(PyExc_Exception, InPyType, *FString::Printf(TEXT("Failed to assign descriptor for '%s'"), UTF8_TO_TCHAR(PropDef->PyGetSet.name)));
+				PyUtil::SetPythonError(PyExc_Exception, PyType, *FString::Printf(TEXT("Failed to assign descriptor for '%s'"), UTF8_TO_TCHAR(PropDef->PyGetSet.name)));
 				return false;
 			}
 		}
 
 		return true;
 	}
+
+	void PrepareOldStructForReinstancing()
+	{
+		check(OldStruct);
+
+		const FString OldStructName = MakeUniqueObjectName(OldStruct->GetOuter(), OldStruct->GetClass(), *FString::Printf(TEXT("%s_REINST"), *OldStruct->GetName())).ToString();
+		OldStruct->SetFlags(RF_NewerVersionExists);
+		OldStruct->ClearFlags(RF_Public | RF_Standalone);
+		OldStruct->Rename(*OldStructName, nullptr, REN_DontCreateRedirectors);
+	}
+
+	FString StructName;
+	PyTypeObject* PyType;
+	UPythonGeneratedStruct* OldStruct;
+	UPythonGeneratedStruct* NewStruct;
 };
 
 void UPythonGeneratedStruct::PostRename(UObject* OldOuter, const FName OldName)
 {
 	Super::PostRename(OldOuter, OldName);
 
-	FPyWrapperTypeRegistry::Get().UnregisterWrappedStructType(OldName, PyType);
-	FPyWrapperTypeRegistry::Get().RegisterWrappedStructType(GetFName(), PyType, !HasAnyFlags(RF_NewerVersionExists));
+	if (PyType)
+	{
+		FPyWrapperTypeRegistry::Get().UnregisterWrappedStructType(OldName, PyType);
+		FPyWrapperTypeRegistry::Get().RegisterWrappedStructType(GetFName(), PyType, !HasAnyFlags(RF_NewerVersionExists));
+	}
 }
 
 void UPythonGeneratedStruct::InitializeStruct(void* Dest, int32 ArrayDim) const
@@ -1434,9 +1498,6 @@ void UPythonGeneratedStruct::InitializeStruct(void* Dest, int32 ArrayDim) const
 
 UPythonGeneratedStruct* UPythonGeneratedStruct::GenerateStruct(PyTypeObject* InPyType)
 {
-	UObject* StructOuter = GetPythonTypeContainer();
-	const FString StructName = PyUtil::GetCleanTypename(InPyType);
-
 	// Get the correct super struct from the parent type in Python
 	UScriptStruct* SuperStruct = nullptr;
 	if (InPyType->tp_base != &PyWrapperStructType)
@@ -1449,20 +1510,8 @@ UPythonGeneratedStruct* UPythonGeneratedStruct::GenerateStruct(PyTypeObject* InP
 		}
 	}
 
-	UPythonGeneratedStruct* OldStruct = FindObject<UPythonGeneratedStruct>(StructOuter, *StructName);
-	if (OldStruct)
-	{
-		FPythonGeneratedStructUtil::PrepareOldStructForReinstancing(OldStruct);
-	}
-
-	UPythonGeneratedStruct* Struct = FPythonGeneratedStructUtil::CreateStruct(StructName, StructOuter, SuperStruct);
-
-	// Get the post-init function
-	Struct->PyPostInitFunction = FPyObjectPtr::StealReference(PyGenUtil::GetPostInitFunc(InPyType));
-	if (!Struct->PyPostInitFunction)
-	{
-		return nullptr;
-	}
+	// Builder used to generate the struct
+	FPythonGeneratedStructBuilder PythonStructBuilder(PyUtil::GetCleanTypename(InPyType), SuperStruct, InPyType);
 
 	// Add the fields to this struct
 	{
@@ -1483,7 +1532,7 @@ UPythonGeneratedStruct* UPythonGeneratedStruct::GenerateStruct(PyTypeObject* InP
 			if (PyObject_IsInstance(FieldValue, (PyObject*)&PyUPropertyDefType) == 1)
 			{
 				FPyUPropertyDef* PyPropDef = (FPyUPropertyDef*)FieldValue;
-				if (!FPythonGeneratedStructUtil::CreatePropertyFromDefinition(Struct, InPyType, FieldName, PyPropDef))
+				if (!PythonStructBuilder.CreatePropertyFromDefinition(FieldName, PyPropDef))
 				{
 					return nullptr;
 				}
@@ -1498,25 +1547,8 @@ UPythonGeneratedStruct* UPythonGeneratedStruct::GenerateStruct(PyTypeObject* InP
 		}
 	}
 
-	// Replace the definitions with real descriptors
-	if (!FPythonGeneratedStructUtil::RegisterDescriptors(Struct, InPyType))
-	{
-		return nullptr;
-	}
-
-	// Let Python know that we've changed its type
-	PyType_Modified(InPyType);
-
-	// Finalize the struct
-	FPythonGeneratedStructUtil::FinalizeStruct(Struct, InPyType);
-
-	// Re-instance the old struct
-	if (OldStruct)
-	{
-		FPyWrapperTypeReinstancer::Get().AddPendingStruct(OldStruct, Struct);
-	}
-
-	return Struct;
+	// Finalize the struct with its post-init function
+	return PythonStructBuilder.Finalize(FPyObjectPtr::StealReference(PyGenUtil::GetPostInitFunc(InPyType)));
 }
 
 #endif	// WITH_PYTHON
diff --git a/Engine/Plugins/Experimental/PythonScriptPlugin/Source/PythonScriptPlugin/Private/PyWrapperStruct.h b/Engine/Plugins/Experimental/PythonScriptPlugin/Source/PythonScriptPlugin/Private/PyWrapperStruct.h
--- a/Engine/Plugins/Experimental/PythonScriptPlugin/Source/PythonScriptPlugin/Private/PyWrapperStruct.h
+++ b/Engine/Plugins/Experimental/PythonScriptPlugin/Source/PythonScriptPlugin/Private/PyWrapperStruct.h
@@ -397,7 +397,7 @@ private:
 	/** Meta-data for this generated struct that is applied to the Python type */
 	FPyWrapperStructMetaData PyMetaData;
 
-	friend struct FPythonGeneratedStructUtil;
+	friend class FPythonGeneratedStructBuilder;
 
 #endif	// WITH_PYTHON
 };
