Guard against some cases that can trigger asserts when constructing object instances from Python

diff --git a/Engine/Plugins/Experimental/PythonScriptPlugin/Source/PythonScriptPlugin/Private/PyWrapperObject.cpp b/Engine/Plugins/Experimental/PythonScriptPlugin/Source/PythonScriptPlugin/Private/PyWrapperObject.cpp
--- a/Engine/Plugins/Experimental/PythonScriptPlugin/Source/PythonScriptPlugin/Private/PyWrapperObject.cpp
+++ b/Engine/Plugins/Experimental/PythonScriptPlugin/Source/PythonScriptPlugin/Private/PyWrapperObject.cpp
@@ -483,15 +483,33 @@ PyTypeObject InitializePyWrapperObjectType()
 					}
 				}
 
-				if (ObjClass->HasAnyClassFlags(CLASS_Abstract))
+				if (ObjClass == UPackage::StaticClass())
 				{
-					PyUtil::SetPythonError(PyExc_Exception, InSelf, *FString::Printf(TEXT("Class '%s' is abstract"), UTF8_TO_TCHAR(Py_TYPE(InSelf)->tp_name)));
+					if (ObjectName.IsNone())
+					{
+						PyUtil::SetPythonError(PyExc_Exception, InSelf, TEXT("Name cannot be 'None' when creating a 'Package'"));
+						return -1;
+					}
+				}
+				else if (!ObjectOuter)
+				{
+					PyUtil::SetPythonError(PyExc_Exception, InSelf, *FString::Printf(TEXT("Outer cannot be null when creating a '%s'"), *ObjClass->GetName()));
 					return -1;
 				}
-				else
+
+				if (ObjectOuter && !ObjectOuter->IsA(ObjClass->ClassWithin))
 				{
-					InitValue = NewObject<UObject>(ObjectOuter, ObjClass, ObjectName);
+					PyUtil::SetPythonError(PyExc_TypeError, InSelf, *FString::Printf(TEXT("Outer '%s' was of type '%s' but must be of type '%s'"), *ObjectOuter->GetPathName(), *ObjectOuter->GetClass()->GetName(), *ObjClass->ClassWithin->GetName()));
+					return -1;
+				}
+
+				if (ObjClass->HasAnyClassFlags(CLASS_Abstract))
+				{
+					PyUtil::SetPythonError(PyExc_Exception, InSelf, *FString::Printf(TEXT("Class '%s' is abstract"), UTF8_TO_TCHAR(Py_TYPE(InSelf)->tp_name)));
+					return -1;
 				}
+				
+				InitValue = NewObject<UObject>(ObjectOuter, ObjClass, ObjectName);
 			}
 			else
 			{
