Don't export Python glue for classes that aren't Blueprintable, and have no script exposed fields

diff --git a/Engine/Plugins/Experimental/PythonScriptPlugin/Source/PythonScriptPlugin/Private/PyGenUtil.cpp b/Engine/Plugins/Experimental/PythonScriptPlugin/Source/PythonScriptPlugin/Private/PyGenUtil.cpp
--- a/Engine/Plugins/Experimental/PythonScriptPlugin/Source/PythonScriptPlugin/Private/PyGenUtil.cpp
+++ b/Engine/Plugins/Experimental/PythonScriptPlugin/Source/PythonScriptPlugin/Private/PyGenUtil.cpp
@@ -1130,7 +1130,7 @@ FString BuildFunctionDocString(const UFunction* InFunc, const FString& InFuncPyt
 	return FunctionDeclDocString;
 }
 
-bool IsBlueprintExposedClass(const UClass* InClass)
+bool IsScriptExposedClass(const UClass* InClass)
 {
 	for (const UClass* ParentClass = InClass; ParentClass; ParentClass = ParentClass->GetSuperClass())
 	{
@@ -1148,7 +1148,7 @@ bool IsBlueprintExposedClass(const UClass* InClass)
 	return false;
 }
 
-bool IsBlueprintExposedStruct(const UScriptStruct* InStruct)
+bool IsScriptExposedStruct(const UScriptStruct* InStruct)
 {
 	for (const UScriptStruct* ParentStruct = InStruct; ParentStruct; ParentStruct = Cast<UScriptStruct>(ParentStruct->GetSuperStruct()))
 	{
@@ -1166,7 +1166,7 @@ bool IsBlueprintExposedStruct(const UScriptStruct* InStruct)
 	return false;
 }
 
-bool IsBlueprintExposedEnum(const UEnum* InEnum)
+bool IsScriptExposedEnum(const UEnum* InEnum)
 {
 	if (InEnum->GetBoolMetaData(BlueprintTypeMetaDataKey))
 	{
@@ -1181,19 +1181,21 @@ bool IsBlueprintExposedEnum(const UEnum* InEnum)
 	return false;
 }
 
-bool IsBlueprintExposedEnumEntry(const UEnum* InEnum, int32 InEnumEntryIndex)
+bool IsScriptExposedEnumEntry(const UEnum* InEnum, int32 InEnumEntryIndex)
 {
 	return !InEnum->HasMetaData(HiddenMetaDataKey, InEnumEntryIndex);
 }
 
-bool IsBlueprintExposedProperty(const UProperty* InProp)
+bool IsScriptExposedProperty(const UProperty* InProp)
 {
-	return InProp->HasAnyPropertyFlags(CPF_BlueprintVisible | CPF_BlueprintAssignable);
+	return !InProp->HasMetaData(ScriptNoExportMetaDataKey) 
+		&& InProp->HasAnyPropertyFlags(CPF_BlueprintVisible | CPF_BlueprintAssignable);
 }
 
-bool IsBlueprintExposedFunction(const UFunction* InFunc)
+bool IsScriptExposedFunction(const UFunction* InFunc)
 {
-	return InFunc->HasAnyFunctionFlags(FUNC_BlueprintCallable | FUNC_BlueprintEvent)
+	return !InFunc->HasMetaData(ScriptNoExportMetaDataKey)
+		&& InFunc->HasAnyFunctionFlags(FUNC_BlueprintCallable | FUNC_BlueprintEvent)
 		&& !InFunc->HasMetaData(BlueprintGetterMetaDataKey)
 		&& !InFunc->HasMetaData(BlueprintSetterMetaDataKey)
 		&& !InFunc->HasMetaData(CustomStructureParamMetaDataKey)
@@ -1201,26 +1203,26 @@ bool IsBlueprintExposedFunction(const UFunction* InFunc)
 		&& !InFunc->HasMetaData(NativeMakeFuncMetaDataKey);
 }
 
-bool IsBlueprintExposedField(const UField* InField)
+bool IsScriptExposedField(const UField* InField)
 {
 	if (const UProperty* Prop = Cast<const UProperty>(InField))
 	{
-		return IsBlueprintExposedProperty(Prop);
+		return IsScriptExposedProperty(Prop);
 	}
 
 	if (const UFunction* Func = Cast<const UFunction>(InField))
 	{
-		return IsBlueprintExposedFunction(Func);
+		return IsScriptExposedFunction(Func);
 	}
 
 	return false;
 }
 
-bool HasBlueprintExposedFields(const UStruct* InStruct)
+bool HasScriptExposedFields(const UStruct* InStruct)
 {
 	for (TFieldIterator<const UField> FieldIt(InStruct); FieldIt; ++FieldIt)
 	{
-		if (IsBlueprintExposedField(*FieldIt))
+		if (IsScriptExposedField(*FieldIt))
 		{
 			return true;
 		}
@@ -1305,28 +1307,28 @@ bool IsDeprecatedFunction(const UFunction* InFunc, FString* OutDeprecationMessag
 
 bool ShouldExportClass(const UClass* InClass)
 {
-	return IsBlueprintExposedClass(InClass) || HasBlueprintExposedFields(InClass);
+	return IsScriptExposedClass(InClass) || HasScriptExposedFields(InClass);
 }
 
 bool ShouldExportStruct(const UScriptStruct* InStruct)
 {
-	return IsBlueprintExposedStruct(InStruct) || HasBlueprintExposedFields(InStruct);
+	return IsScriptExposedStruct(InStruct) || HasScriptExposedFields(InStruct);
 }
 
 bool ShouldExportEnum(const UEnum* InEnum)
 {
-	return IsBlueprintExposedEnum(InEnum);
+	return IsScriptExposedEnum(InEnum);
 }
 
 bool ShouldExportEnumEntry(const UEnum* InEnum, int32 InEnumEntryIndex)
 {
-	return IsBlueprintExposedEnumEntry(InEnum, InEnumEntryIndex);
+	return IsScriptExposedEnumEntry(InEnum, InEnumEntryIndex);
 }
 
 bool ShouldExportProperty(const UProperty* InProp)
 {
-	const bool bCanScriptExport = !InProp->HasMetaData(ScriptNoExportMetaDataKey);
-	return bCanScriptExport && (IsBlueprintExposedProperty(InProp) || IsDeprecatedProperty(InProp));
+	const bool bCanScriptExport = !InProp->HasMetaData(ScriptNoExportMetaDataKey); // Need to test this again here as IsScriptExposedProperty checks it internally, but IsDeprecatedProperty doesn't
+	return bCanScriptExport && (IsScriptExposedProperty(InProp) || IsDeprecatedProperty(InProp));
 }
 
 bool ShouldExportEditorOnlyProperty(const UProperty* InProp)
@@ -1337,8 +1339,7 @@ bool ShouldExportEditorOnlyProperty(const UProperty* InProp)
 
 bool ShouldExportFunction(const UFunction* InFunc)
 {
-	const bool bCanScriptExport = !InFunc->HasMetaData(ScriptNoExportMetaDataKey);
-	return bCanScriptExport && IsBlueprintExposedFunction(InFunc);
+	return IsScriptExposedFunction(InFunc);
 }
 
 FString PythonizeName(const FString& InName, const EPythonizeNameCase InNameCase)
