Hide BlueprintInternalUseOnly and CustomThunk functions from Python

BlueprintInternalUseOnly functions are usually designed to be called via a K2 node wrapper, which may provide additional sanitization or data transformation (which Python is unaware of, so cannot mimic).

CustomThunk functions rely on popping Bytecode data off the stack, but Python calls functions natively so there is no Bytecode available and these functions will crash. We hope to address CustomThunk support in the future by synthesizing the correct Bytecode.

diff --git a/Engine/Plugins/Experimental/PythonScriptPlugin/Source/PythonScriptPlugin/Private/PyGenUtil.cpp b/Engine/Plugins/Experimental/PythonScriptPlugin/Source/PythonScriptPlugin/Private/PyGenUtil.cpp
--- a/Engine/Plugins/Experimental/PythonScriptPlugin/Source/PythonScriptPlugin/Private/PyGenUtil.cpp
+++ b/Engine/Plugins/Experimental/PythonScriptPlugin/Source/PythonScriptPlugin/Private/PyGenUtil.cpp
@@ -57,10 +57,11 @@ const FName NotBlueprintTypeMetaDataKey = TEXT("NotBlueprintType");
 const FName BlueprintSpawnableComponentMetaDataKey = TEXT("BlueprintSpawnableComponent");
 const FName BlueprintGetterMetaDataKey = TEXT("BlueprintGetter");
 const FName BlueprintSetterMetaDataKey = TEXT("BlueprintSetter");
+const FName BlueprintInternalUseOnlyMetaDataKey = TEXT("BlueprintInternalUseOnly");
+const FName CustomThunkMetaDataKey = TEXT("CustomThunk");
 const FName DeprecatedPropertyMetaDataKey = TEXT("DeprecatedProperty");
 const FName DeprecatedFunctionMetaDataKey = TEXT("DeprecatedFunction");
 const FName DeprecationMessageMetaDataKey = TEXT("DeprecationMessage");
-const FName CustomStructureParamMetaDataKey = TEXT("CustomStructureParam");
 const FName HasNativeMakeMetaDataKey = TEXT("HasNativeMake");
 const FName HasNativeBreakMetaDataKey = TEXT("HasNativeBreak");
 const FName NativeBreakFuncMetaDataKey = TEXT("NativeBreakFunc");
@@ -1198,7 +1199,8 @@ bool IsScriptExposedFunction(const UFunction* InFunc)
 		&& InFunc->HasAnyFunctionFlags(FUNC_BlueprintCallable | FUNC_BlueprintEvent)
 		&& !InFunc->HasMetaData(BlueprintGetterMetaDataKey)
 		&& !InFunc->HasMetaData(BlueprintSetterMetaDataKey)
-		&& !InFunc->HasMetaData(CustomStructureParamMetaDataKey)
+		&& !InFunc->HasMetaData(BlueprintInternalUseOnlyMetaDataKey)
+		&& !InFunc->HasMetaData(CustomThunkMetaDataKey)
 		&& !InFunc->HasMetaData(NativeBreakFuncMetaDataKey)
 		&& !InFunc->HasMetaData(NativeMakeFuncMetaDataKey);
 }
diff --git a/Engine/Plugins/Experimental/PythonScriptPlugin/Source/PythonScriptPlugin/Private/PyGenUtil.h b/Engine/Plugins/Experimental/PythonScriptPlugin/Source/PythonScriptPlugin/Private/PyGenUtil.h
--- a/Engine/Plugins/Experimental/PythonScriptPlugin/Source/PythonScriptPlugin/Private/PyGenUtil.h
+++ b/Engine/Plugins/Experimental/PythonScriptPlugin/Source/PythonScriptPlugin/Private/PyGenUtil.h
@@ -46,7 +46,8 @@ namespace PyGenUtil
 	extern const FName BlueprintSpawnableComponentMetaDataKey;
 	extern const FName BlueprintGetterMetaDataKey;
 	extern const FName BlueprintSetterMetaDataKey;
-	extern const FName CustomStructureParamMetaDataKey;
+	extern const FName BlueprintInternalUseOnlyMetaDataKey;
+	extern const FName CustomThunkMetaDataKey;
 	extern const FName HasNativeMakeMetaDataKey;
 	extern const FName HasNativeBreakMetaDataKey;
 	extern const FName NativeBreakFuncMetaDataKey;
