Added an interface for UObject types that own Python data so that they can be cleared during interpreter shutdown if the UObject is still externally referenced and couldn't be GC'd

diff --git a/Engine/Plugins/Experimental/PythonScriptPlugin/Source/PythonScriptPlugin/Private/PyReferenceCollector.cpp b/Engine/Plugins/Experimental/PythonScriptPlugin/Source/PythonScriptPlugin/Private/PyReferenceCollector.cpp
--- a/Engine/Plugins/Experimental/PythonScriptPlugin/Source/PythonScriptPlugin/Private/PyReferenceCollector.cpp
+++ b/Engine/Plugins/Experimental/PythonScriptPlugin/Source/PythonScriptPlugin/Private/PyReferenceCollector.cpp
@@ -73,120 +73,103 @@ void FPyReferenceCollector::PurgeUnrealObjectReferences(const TArrayView<const U
 
 void FPyReferenceCollector::PurgeUnrealGeneratedTypes()
 {
-	auto RunPurgeUnrealGeneratedTypes = [this](const bool bLogFailure) -> bool
+	FPurgingReferenceCollector PurgingReferenceCollector;
+	TArray<FWeakObjectPtr> WeakReferencesToPurgedObjects;
+
+	auto FlagObjectForPurge = [&PurgingReferenceCollector, &WeakReferencesToPurgedObjects](UObject* InObject, const bool bMarkPendingKill)
 	{
-		FPurgingReferenceCollector PurgingReferenceCollector;
-		TArray<FWeakObjectPtr> WeakReferencesToPurgedObjects;
+		check(!InObject->HasAnyInternalFlags(EInternalObjectFlags::Native));
 
-		auto FlagObjectForPurge = [&PurgingReferenceCollector, &WeakReferencesToPurgedObjects](UObject* InObject, const bool bMarkPendingKill)
+		if (InObject->IsRooted())
 		{
-			if (!InObject->HasAnyInternalFlags(EInternalObjectFlags::Native) && !InObject->HasAnyFlags(RF_ClassDefaultObject))
-			{
-				if (InObject->IsRooted())
-				{
-					InObject->RemoveFromRoot();
-				}
-				InObject->ClearFlags(RF_Public | RF_Standalone);
-				if (bMarkPendingKill)
-				{
-					InObject->MarkPendingKill();
-				}
-
-				WeakReferencesToPurgedObjects.Add(InObject);
-			}
-
-			PurgingReferenceCollector.AddObjectToPurge(InObject);
-		};
+			InObject->RemoveFromRoot();
+		}
+		InObject->ClearFlags(RF_Public | RF_Standalone);
 
-		// Clean-up Python generated class types and instances
-		// The class types are instances of UPythonGeneratedClass
+		if (!InObject->HasAnyFlags(RF_ClassDefaultObject))
 		{
-			ForEachObjectOfClass(UPythonGeneratedClass::StaticClass(), [&FlagObjectForPurge](UObject* InObject)
+			if (bMarkPendingKill)
 			{
-				UPythonGeneratedClass* PythonGeneratedClass = CastChecked<UPythonGeneratedClass>(InObject);
-				bool bMarkClassPendingKill = false; // Mark types as PendingKill if they have no instances left (excluding their CDO)
-
-				ForEachObjectOfClass(PythonGeneratedClass, [&bMarkClassPendingKill, &FlagObjectForPurge](UObject* InInnerObject)
-				{
-					bMarkClassPendingKill &= InInnerObject->HasAnyFlags(RF_ClassDefaultObject);
-					FlagObjectForPurge(InInnerObject, /*bMarkPendingKill*/true);
-				}, false);
-
-				FlagObjectForPurge(PythonGeneratedClass, bMarkClassPendingKill);
-			}, false);
+				InObject->MarkPendingKill();
+			}
+			WeakReferencesToPurgedObjects.Add(InObject);
 		}
 
-		// Clean-up Python generated struct types
-		// The struct types are instances of UPythonGeneratedStruct
+		PurgingReferenceCollector.AddObjectToPurge(InObject);
+	};
+
+	// Clean-up Python generated class types and instances
+	// The class types are instances of UPythonGeneratedClass
+	{
+		ForEachObjectOfClass(UPythonGeneratedClass::StaticClass(), [&FlagObjectForPurge](UObject* InObject)
 		{
-			ForEachObjectOfClass(UPythonGeneratedStruct::StaticClass(), [&FlagObjectForPurge](UObject* InObject)
+			FlagObjectForPurge(InObject, /*bMarkPendingKill*/false);
+
+			UPythonGeneratedClass* PythonGeneratedClass = CastChecked<UPythonGeneratedClass>(InObject);
+			ForEachObjectOfClass(PythonGeneratedClass, [&FlagObjectForPurge](UObject* InInnerObject)
 			{
-				FlagObjectForPurge(InObject, /*bMarkPendingKill*/false);
-			}, false);
-		}
+				FlagObjectForPurge(InInnerObject, /*bMarkPendingKill*/true);
+			}, false, RF_NoFlags, EInternalObjectFlags::Native);
+		}, false, RF_ClassDefaultObject, EInternalObjectFlags::Native);
+	}
 
-		// Clean-up Python generated enum types
-		// The enum types are instances of UPythonGeneratedEnum
+	// Clean-up Python generated struct types
+	// The struct types are instances of UPythonGeneratedStruct
+	{
+		ForEachObjectOfClass(UPythonGeneratedStruct::StaticClass(), [&FlagObjectForPurge](UObject* InObject)
 		{
-			ForEachObjectOfClass(UPythonGeneratedEnum::StaticClass(), [&FlagObjectForPurge](UObject* InObject)
-			{
-				FlagObjectForPurge(InObject, /*bMarkPendingKill*/false);
-			}, false);
-		}
+			FlagObjectForPurge(InObject, /*bMarkPendingKill*/false);
+		}, false, RF_ClassDefaultObject, EInternalObjectFlags::Native);
+	}
 
-		// Clean-up Python callable types and instances
-		// The callable types all derive directly from UPythonCallableForDelegate
+	// Clean-up Python generated enum types
+	// The enum types are instances of UPythonGeneratedEnum
+	{
+		ForEachObjectOfClass(UPythonGeneratedEnum::StaticClass(), [&FlagObjectForPurge](UObject* InObject)
 		{
-			TArray<UClass*> PythonCallableClasses;
-			GetDerivedClasses(UPythonCallableForDelegate::StaticClass(), PythonCallableClasses, true);
+			FlagObjectForPurge(InObject, /*bMarkPendingKill*/false);
+		}, false, RF_ClassDefaultObject, EInternalObjectFlags::Native);
+	}
 
-			for (UClass* PythonCallableClass : PythonCallableClasses)
-			{
-				bool bMarkClassPendingKill = false; // Mark types as PendingKill if they have no instances left (excluding their CDO)
+	// Clean-up Python callable types and instances
+	// The callable types all derive directly from UPythonCallableForDelegate
+	{
+		TArray<UClass*> PythonCallableClasses;
+		GetDerivedClasses(UPythonCallableForDelegate::StaticClass(), PythonCallableClasses, true);
 
-				ForEachObjectOfClass(PythonCallableClass, [&bMarkClassPendingKill, &FlagObjectForPurge](UObject* InObject)
-				{
-					bMarkClassPendingKill &= InObject->HasAnyFlags(RF_ClassDefaultObject);
-					FlagObjectForPurge(InObject, /*bMarkPendingKill*/true);
-				}, false);
+		for (UClass* PythonCallableClass : PythonCallableClasses)
+		{
+			// The classes themselves don't hold Python resources, so it's not imperative that 
+			// they are destroyed, but we don't need them anymore so attempt to let them die
+			PythonCallableClass->ClearFlags(RF_Public | RF_Standalone);
 
-				FlagObjectForPurge(PythonCallableClass, bMarkClassPendingKill);
-			}
+			ForEachObjectOfClass(PythonCallableClass, [&FlagObjectForPurge](UObject* InObject)
+			{
+				FlagObjectForPurge(InObject, /*bMarkPendingKill*/true);
+			}, false, RF_NoFlags, EInternalObjectFlags::Native);
 		}
+	}
 
-		if (PurgingReferenceCollector.HasObjectToPurge())
-		{
-			AddReferencedObjects(PurgingReferenceCollector);
-			CollectGarbage(GARBAGE_COLLECTION_KEEPFLAGS);
+	if (PurgingReferenceCollector.HasObjectToPurge())
+	{
+		AddReferencedObjects(PurgingReferenceCollector);
+		CollectGarbage(GARBAGE_COLLECTION_KEEPFLAGS);
 
-			bool bHasLeftoverObjects = false;
-			for (const FWeakObjectPtr& WeakReferencesToPurgedObject : WeakReferencesToPurgedObjects)
+		for (const FWeakObjectPtr& WeakReferencesToPurgedObject : WeakReferencesToPurgedObjects)
+		{
+			if (UObject* FailedToPurgeObject = WeakReferencesToPurgedObject.Get())
 			{
-				if (UObject* FailedToPurgeObject = WeakReferencesToPurgedObject.Get())
+				if (IPythonResourceOwner* PythonResourceOwner = Cast<IPythonResourceOwner>(FailedToPurgeObject))
 				{
-					bHasLeftoverObjects = true;
-					if (bLogFailure)
-					{
-						UE_LOG(LogPython, Warning, TEXT("Object '%s' failed to purge when requested by PurgeUnrealGeneratedTypes. This may lead to crashes!"), *FailedToPurgeObject->GetPathName());
-					}
-					else if (bHasLeftoverObjects)
-					{
-						break;
-					}
+					UE_LOG(LogPython, Display, TEXT("Object '%s' was externally referenced when shutting down Python. Forcibly releasing its Python resources!"), *FailedToPurgeObject->GetPathName());
+					PythonResourceOwner->ReleasePythonResources();
+				}
+				else
+				{
+					UE_LOG(LogPython, Warning, TEXT("Object '%s' was externally referenced when shutting down Python and could not gracefully cleaned-up. This may lead to crashes!"), *FailedToPurgeObject->GetPathName());
 				}
 			}
-			return bHasLeftoverObjects;
 		}
-
-		return false;
-	};
-
-	// We run two purge passes:
-	//  - Pass 1 will force purge any type instances, and any types that have no instances
-	//  - Pass 2 will purge any types that no longer have instances left after running pass 1
-	if (RunPurgeUnrealGeneratedTypes(/*bLogFailure*/false))
-	{
-		RunPurgeUnrealGeneratedTypes(/*bLogFailure*/false);
 	}
 }
 
diff --git a/Engine/Plugins/Experimental/PythonScriptPlugin/Source/PythonScriptPlugin/Private/PyWrapperBase.cpp b/Engine/Plugins/Experimental/PythonScriptPlugin/Source/PythonScriptPlugin/Private/PyWrapperBase.cpp
--- a/Engine/Plugins/Experimental/PythonScriptPlugin/Source/PythonScriptPlugin/Private/PyWrapperBase.cpp
+++ b/Engine/Plugins/Experimental/PythonScriptPlugin/Source/PythonScriptPlugin/Private/PyWrapperBase.cpp
@@ -144,3 +144,8 @@ FPyWrapperBaseMetaData* FPyWrapperBaseMetaData::GetMetaData(FPyWrapperBase* Inst
 }
 
 #endif	// WITH_PYTHON
+
+UPythonResourceOwner::UPythonResourceOwner(const FObjectInitializer& ObjectInitializer)
+	: Super(ObjectInitializer)
+{
+}
diff --git a/Engine/Plugins/Experimental/PythonScriptPlugin/Source/PythonScriptPlugin/Private/PyWrapperBase.h b/Engine/Plugins/Experimental/PythonScriptPlugin/Source/PythonScriptPlugin/Private/PyWrapperBase.h
--- a/Engine/Plugins/Experimental/PythonScriptPlugin/Source/PythonScriptPlugin/Private/PyWrapperBase.h
+++ b/Engine/Plugins/Experimental/PythonScriptPlugin/Source/PythonScriptPlugin/Private/PyWrapperBase.h
@@ -9,6 +9,7 @@
 #include "PyConversionResult.h"
 #include "Misc/Guid.h"
 #include "UObject/GCObject.h"
+#include "PyWrapperBase.generated.h"
 
 #if WITH_PYTHON
 
@@ -85,3 +86,21 @@ struct FPyWrapperBaseMetaData
 typedef TPyPtr<FPyWrapperBase> FPyWrapperBasePtr;
 
 #endif	// WITH_PYTHON
+
+UINTERFACE(MinimalApi)
+class UPythonResourceOwner : public UInterface
+{
+	GENERATED_UINTERFACE_BODY()
+};
+
+class IPythonResourceOwner
+{
+	GENERATED_IINTERFACE_BODY()
+
+public:
+	/**
+	 * Release all Python resources owned by this object.
+	 * Called during Python shutdown to free resources from lingering UObject-based instances.
+	 */
+	virtual void ReleasePythonResources() = 0;
+};
diff --git a/Engine/Plugins/Experimental/PythonScriptPlugin/Source/PythonScriptPlugin/Private/PyWrapperDelegate.cpp b/Engine/Plugins/Experimental/PythonScriptPlugin/Source/PythonScriptPlugin/Private/PyWrapperDelegate.cpp
--- a/Engine/Plugins/Experimental/PythonScriptPlugin/Source/PythonScriptPlugin/Private/PyWrapperDelegate.cpp
+++ b/Engine/Plugins/Experimental/PythonScriptPlugin/Source/PythonScriptPlugin/Private/PyWrapperDelegate.cpp
@@ -15,25 +15,6 @@
 
 const FName UPythonCallableForDelegate::GeneratedFuncName = "CallPython";
 
-void UPythonCallableForDelegate::BeginDestroy()
-{
-#if WITH_PYTHON
-	// This may be called after Python has already shut down
-	if (Py_IsInitialized())
-	{
-		FPyScopedGIL GIL;
-		PyCallable.Reset();
-	}
-	else
-	{
-		// Release ownership if Python has been shut down to avoid attempting to delete the callable (which is already dead)
-		PyCallable.Release();
-	}
-#endif	// WITH_PYTHON
-
-	Super::BeginDestroy();
-}
-
 DEFINE_FUNCTION(UPythonCallableForDelegate::CallPythonNative)
 {
 #if WITH_PYTHON
@@ -118,6 +99,28 @@ DEFINE_FUNCTION(UPythonCallableForDelegate::CallPythonNative)
 }
 
 #if WITH_PYTHON
+void UPythonCallableForDelegate::BeginDestroy()
+{
+	// This may be called after Python has already shut down
+	if (Py_IsInitialized())
+	{
+		FPyScopedGIL GIL;
+		PyCallable.Reset();
+	}
+	else
+	{
+		// Release ownership if Python has been shut down to avoid attempting to delete the callable (which is already dead)
+		PyCallable.Release();
+	}
+
+	Super::BeginDestroy();
+}
+
+void UPythonCallableForDelegate::ReleasePythonResources()
+{
+	PyCallable.Reset();
+}
+
 PyObject* UPythonCallableForDelegate::GetCallable() const
 {
 	return (PyObject*)PyCallable.GetPtr();
diff --git a/Engine/Plugins/Experimental/PythonScriptPlugin/Source/PythonScriptPlugin/Private/PyWrapperDelegate.h b/Engine/Plugins/Experimental/PythonScriptPlugin/Source/PythonScriptPlugin/Private/PyWrapperDelegate.h
--- a/Engine/Plugins/Experimental/PythonScriptPlugin/Source/PythonScriptPlugin/Private/PyWrapperDelegate.h
+++ b/Engine/Plugins/Experimental/PythonScriptPlugin/Source/PythonScriptPlugin/Private/PyWrapperDelegate.h
@@ -13,32 +13,42 @@
  * @note This can't go inside the WITH_PYTHON block due to UHT parsing limitations (it doesn't understand that macro)
  */
 UCLASS()
-class UPythonCallableForDelegate : public UObject
+class UPythonCallableForDelegate : public UObject, public IPythonResourceOwner
 {
 	GENERATED_BODY()
 
 public:
-	//~ UObject interface
-	virtual void BeginDestroy() override;
+	/** Name given to the generated function that we should bind to the Unreal delegate */
+	static const FName GeneratedFuncName;
 
 	/** Native function implementation used by the signature correct Unreal functions added to the derived classes (the ones that are bound to the delegate itself) */
 	DECLARE_FUNCTION(CallPythonNative);
 
 #if WITH_PYTHON
+
+public:
+	//~ UObject interface
+	virtual void BeginDestroy() override;
+
+	//~ IPythonResourceOwner interface
+	virtual void ReleasePythonResources() override;
+
 	/** Get the Python callable object on this instance (borrowed reference) */
 	PyObject* GetCallable() const;
 
 	/** Set the Python callable object on this instance */
 	void SetCallable(PyObject* InCallable);
-#endif	// WITH_PYTHON
-
-	/** Name given to the generated function that we should bind to the Unreal delegate */
-	static const FName GeneratedFuncName;
 
 private:
-#if WITH_PYTHON
 	/** The callable Python callable object this object wraps (if any) */
 	FPyObjectPtr PyCallable;
+
+#else	// WITH_PYTHON
+
+public:
+	//~ IPythonResourceOwner interface
+	virtual void ReleasePythonResources() override {}
+
 #endif	// WITH_PYTHON
 };
 
diff --git a/Engine/Plugins/Experimental/PythonScriptPlugin/Source/PythonScriptPlugin/Private/PyWrapperEnum.cpp b/Engine/Plugins/Experimental/PythonScriptPlugin/Source/PythonScriptPlugin/Private/PyWrapperEnum.cpp
--- a/Engine/Plugins/Experimental/PythonScriptPlugin/Source/PythonScriptPlugin/Private/PyWrapperEnum.cpp
+++ b/Engine/Plugins/Experimental/PythonScriptPlugin/Source/PythonScriptPlugin/Private/PyWrapperEnum.cpp
@@ -839,6 +839,13 @@ private:
 	UPythonGeneratedEnum* NewEnum;
 };
 
+void UPythonGeneratedEnum::ReleasePythonResources()
+{
+	PyType.Reset();
+	EnumValueDefs.Reset();
+	PyMetaData = FPyWrapperEnumMetaData();
+}
+
 UPythonGeneratedEnum* UPythonGeneratedEnum::GenerateEnum(PyTypeObject* InPyType)
 {
 	// Builder used to generate the enum
diff --git a/Engine/Plugins/Experimental/PythonScriptPlugin/Source/PythonScriptPlugin/Private/PyWrapperEnum.h b/Engine/Plugins/Experimental/PythonScriptPlugin/Source/PythonScriptPlugin/Private/PyWrapperEnum.h
--- a/Engine/Plugins/Experimental/PythonScriptPlugin/Source/PythonScriptPlugin/Private/PyWrapperEnum.h
+++ b/Engine/Plugins/Experimental/PythonScriptPlugin/Source/PythonScriptPlugin/Private/PyWrapperEnum.h
@@ -110,13 +110,16 @@ typedef TPyPtr<FPyWrapperEnum> FPyWrapperEnumPtr;
 
 /** An Unreal enum that was generated from a Python type */
 UCLASS()
-class UPythonGeneratedEnum : public UEnum
+class UPythonGeneratedEnum : public UEnum, public IPythonResourceOwner
 {
 	GENERATED_BODY()
 
 #if WITH_PYTHON
 
 public:
+	//~ IPythonResourceOwner interface
+	virtual void ReleasePythonResources() override;
+
 	/** Generate an Unreal enum from the given Python type */
 	static UPythonGeneratedEnum* GenerateEnum(PyTypeObject* InPyType);
 
@@ -145,5 +148,11 @@ private:
 
 	friend class FPythonGeneratedEnumBuilder;
 
+#else	// WITH_PYTHON
+
+public:
+	//~ IPythonResourceOwner interface
+	virtual void ReleasePythonResources() override {}
+
 #endif	// WITH_PYTHON
 };
diff --git a/Engine/Plugins/Experimental/PythonScriptPlugin/Source/PythonScriptPlugin/Private/PyWrapperObject.cpp b/Engine/Plugins/Experimental/PythonScriptPlugin/Source/PythonScriptPlugin/Private/PyWrapperObject.cpp
--- a/Engine/Plugins/Experimental/PythonScriptPlugin/Source/PythonScriptPlugin/Private/PyWrapperObject.cpp
+++ b/Engine/Plugins/Experimental/PythonScriptPlugin/Source/PythonScriptPlugin/Private/PyWrapperObject.cpp
@@ -1687,6 +1687,15 @@ void UPythonGeneratedClass::PostInitInstance(UObject* InObj)
 	}
 }
 
+void UPythonGeneratedClass::ReleasePythonResources()
+{
+	PyType.Reset();
+	PyPostInitFunction.Reset();
+	PropertyDefs.Reset();
+	FunctionDefs.Reset();
+	PyMetaData = FPyWrapperObjectMetaData();
+}
+
 UPythonGeneratedClass* UPythonGeneratedClass::GenerateClass(PyTypeObject* InPyType)
 {
 	// Get the correct super class from the parent type in Python
diff --git a/Engine/Plugins/Experimental/PythonScriptPlugin/Source/PythonScriptPlugin/Private/PyWrapperObject.h b/Engine/Plugins/Experimental/PythonScriptPlugin/Source/PythonScriptPlugin/Private/PyWrapperObject.h
--- a/Engine/Plugins/Experimental/PythonScriptPlugin/Source/PythonScriptPlugin/Private/PyWrapperObject.h
+++ b/Engine/Plugins/Experimental/PythonScriptPlugin/Source/PythonScriptPlugin/Private/PyWrapperObject.h
@@ -180,7 +180,7 @@ typedef TPyPtr<FPyWrapperObject> FPyWrapperObjectPtr;
 
 /** An Unreal class that was generated from a Python type */
 UCLASS()
-class UPythonGeneratedClass : public UClass
+class UPythonGeneratedClass : public UClass, public IPythonResourceOwner
 {
 	GENERATED_BODY()
 
@@ -193,6 +193,9 @@ public:
 	//~ UClass interface
 	virtual void PostInitInstance(UObject* InObj) override;
 
+	//~ IPythonResourceOwner interface
+	virtual void ReleasePythonResources() override;
+
 	/** Generate an Unreal class from the given Python type */
 	static UPythonGeneratedClass* GenerateClass(PyTypeObject* InPyType);
 
@@ -223,5 +226,11 @@ private:
 
 	friend class FPythonGeneratedClassBuilder;
 
+#else	// WITH_PYTHON
+
+public:
+	//~ IPythonResourceOwner interface
+	virtual void ReleasePythonResources() override {}
+
 #endif	// WITH_PYTHON
 };
diff --git a/Engine/Plugins/Experimental/PythonScriptPlugin/Source/PythonScriptPlugin/Private/PyWrapperStruct.cpp b/Engine/Plugins/Experimental/PythonScriptPlugin/Source/PythonScriptPlugin/Private/PyWrapperStruct.cpp
--- a/Engine/Plugins/Experimental/PythonScriptPlugin/Source/PythonScriptPlugin/Private/PyWrapperStruct.cpp
+++ b/Engine/Plugins/Experimental/PythonScriptPlugin/Source/PythonScriptPlugin/Private/PyWrapperStruct.cpp
@@ -1500,6 +1500,14 @@ void UPythonGeneratedStruct::InitializeStruct(void* Dest, int32 ArrayDim) const
 	}
 }
 
+void UPythonGeneratedStruct::ReleasePythonResources()
+{
+	PyType.Reset();
+	PyPostInitFunction.Reset();
+	PropertyDefs.Reset();
+	PyMetaData = FPyWrapperStructMetaData();
+}
+
 UPythonGeneratedStruct* UPythonGeneratedStruct::GenerateStruct(PyTypeObject* InPyType)
 {
 	// Get the correct super struct from the parent type in Python
diff --git a/Engine/Plugins/Experimental/PythonScriptPlugin/Source/PythonScriptPlugin/Private/PyWrapperStruct.h b/Engine/Plugins/Experimental/PythonScriptPlugin/Source/PythonScriptPlugin/Private/PyWrapperStruct.h
--- a/Engine/Plugins/Experimental/PythonScriptPlugin/Source/PythonScriptPlugin/Private/PyWrapperStruct.h
+++ b/Engine/Plugins/Experimental/PythonScriptPlugin/Source/PythonScriptPlugin/Private/PyWrapperStruct.h
@@ -368,7 +368,7 @@ typedef TPyPtr<FPyWrapperStruct> FPyWrapperStructPtr;
 
 /** An Unreal struct that was generated from a Python type */
 UCLASS()
-class UPythonGeneratedStruct : public UScriptStruct
+class UPythonGeneratedStruct : public UScriptStruct, public IPythonResourceOwner
 {
 	GENERATED_BODY()
 
@@ -381,6 +381,9 @@ public:
 	//~ UStruct interface
 	virtual void InitializeStruct(void* Dest, int32 ArrayDim = 1) const override;
 
+	//~ IPythonResourceOwner interface
+	virtual void ReleasePythonResources() override;
+
 	/** Generate an Unreal struct from the given Python type */
 	static UPythonGeneratedStruct* GenerateStruct(PyTypeObject* InPyType);
 
@@ -399,5 +402,11 @@ private:
 
 	friend class FPythonGeneratedStructBuilder;
 
+#else	// WITH_PYTHON
+
+public:
+	//~ IPythonResourceOwner interface
+	virtual void ReleasePythonResources() override {}
+
 #endif	// WITH_PYTHON
 };
