Added missing include

#rnx

diff --git a/Engine/Plugins/Experimental/PythonScriptPlugin/Source/PythonScriptPlugin/Private/PyWrapperBase.h b/Engine/Plugins/Experimental/PythonScriptPlugin/Source/PythonScriptPlugin/Private/PyWrapperBase.h
--- a/Engine/Plugins/Experimental/PythonScriptPlugin/Source/PythonScriptPlugin/Private/PyWrapperBase.h
+++ b/Engine/Plugins/Experimental/PythonScriptPlugin/Source/PythonScriptPlugin/Private/PyWrapperBase.h
@@ -9,6 +9,7 @@
 #include "PyConversionResult.h"
 #include "Misc/Guid.h"
 #include "UObject/GCObject.h"
+#include "UObject/Interface.h"
 #include "PyWrapperBase.generated.h"
 
 #if WITH_PYTHON
