Fix typo (missing ellipsis) in PythonScriptPlugin menu entry (File > Execute Python Script)

diff --git a/Engine/Plugins/Experimental/PythonScriptPlugin/Source/PythonScriptPlugin/Private/PythonScriptPlugin.cpp b/Engine/Plugins/Experimental/PythonScriptPlugin/Source/PythonScriptPlugin/Private/PythonScriptPlugin.cpp
--- a/Engine/Plugins/Experimental/PythonScriptPlugin/Source/PythonScriptPlugin/Private/PythonScriptPlugin.cpp
+++ b/Engine/Plugins/Experimental/PythonScriptPlugin/Source/PythonScriptPlugin/Private/PythonScriptPlugin.cpp
@@ -339,7 +339,7 @@ private:
 	{
 		MenuBuilder.BeginSection("Python", LOCTEXT("Python", "Python"));
 		MenuBuilder.AddMenuEntry(
-			LOCTEXT("OpenPython", "Execute Python Script"),
+			LOCTEXT("OpenPython", "Execute Python Script..."),
 			LOCTEXT("OpenPythonTooltip", "Open a Python Script file and Execute it."),
 			FSlateIcon(),
 			FUIAction(FExecuteAction::CreateRaw(this, &FPythonCommandMenuImpl::Menu_ExecutePython))
