Fixed crash calling Python UObject function override with no args

#rnx

diff --git a/Engine/Plugins/Experimental/PythonScriptPlugin/Source/PythonScriptPlugin/Private/PyGenUtil.cpp b/Engine/Plugins/Experimental/PythonScriptPlugin/Source/PythonScriptPlugin/Private/PyGenUtil.cpp
--- a/Engine/Plugins/Experimental/PythonScriptPlugin/Source/PythonScriptPlugin/Private/PyGenUtil.cpp
+++ b/Engine/Plugins/Experimental/PythonScriptPlugin/Source/PythonScriptPlugin/Private/PyGenUtil.cpp
@@ -1026,7 +1026,7 @@ bool UnpackReturnValues(PyObject* InRetVals, const FOutParmRec* InOutputParms, c
 bool InvokePythonCallableFromUnrealFunctionThunk(FPyObjectPtr InSelf, PyObject* InCallable, const UFunction* InFunc, UObject* Context, FFrame& Stack, RESULT_DECL)
 {
 	// Allocate memory to store our local argument data
-	void* LocalStruct = FMemory_Alloca(InFunc->GetStructureSize());
+	void* LocalStruct = FMemory_Alloca(FMath::Max<int32>(1, InFunc->GetStructureSize()));
 	InFunc->InitializeStruct(LocalStruct);
 	ON_SCOPE_EXIT
 	{
