Put back CanContainContent: true in PythonScriptPlugin as it's required to parse the python scripts in the content folder.

#rnx

diff --git a/Engine/Plugins/Experimental/PythonScriptPlugin/PythonScriptPlugin.uplugin b/Engine/Plugins/Experimental/PythonScriptPlugin/PythonScriptPlugin.uplugin
--- a/Engine/Plugins/Experimental/PythonScriptPlugin/PythonScriptPlugin.uplugin
+++ b/Engine/Plugins/Experimental/PythonScriptPlugin/PythonScriptPlugin.uplugin
@@ -11,7 +11,7 @@
 	"MarketplaceURL" : "",
 	"SupportURL" : "https://forums.unrealengine.com/",
 	"EnabledByDefault" : false,
-	"CanContainContent" : false,
+	"CanContainContent" : true,
 	"IsBetaVersion" : true,
 	"Installed" : false,
 	
