Fix crash building a monolithic editor executable due to global constructors creating UObject instances before the UObject subsystem has been initialized.

diff --git a/Engine/Plugins/Experimental/PythonScriptPlugin/Source/PythonScriptPlugin/Private/PyCore.cpp b/Engine/Plugins/Experimental/PythonScriptPlugin/Source/PythonScriptPlugin/Private/PyCore.cpp
--- a/Engine/Plugins/Experimental/PythonScriptPlugin/Source/PythonScriptPlugin/Private/PyCore.cpp
+++ b/Engine/Plugins/Experimental/PythonScriptPlugin/Source/PythonScriptPlugin/Private/PyCore.cpp
@@ -80,17 +80,26 @@ void ApplyMetaData(PyObject* InMetaData, const TFunctionRef<void(const FString&,
 
 } // namespace PyCoreUtil
 
-TStrongObjectPtr<UStruct> GPythonPropertyContainer;
-TStrongObjectPtr<UPackage> GPythonTypeContainer;
+TStrongObjectPtr<UStruct>& GetPythonPropertyContainerSingleton()
+{
+	static TStrongObjectPtr<UStruct> PythonPropertyContainer;
+	return PythonPropertyContainer;
+}
+
+TStrongObjectPtr<UPackage>& GetPythonTypeContainerSingleton()
+{
+	static TStrongObjectPtr<UPackage> PythonTypeContainer;
+	return PythonTypeContainer;
+}
 
 UObject* GetPythonPropertyContainer()
 {
-	return GPythonPropertyContainer.Get();
+	return GetPythonPropertyContainerSingleton().Get();
 }
 
 UObject* GetPythonTypeContainer()
 {
-	return GPythonTypeContainer.Get();
+	return GetPythonTypeContainerSingleton().Get();
 }
 
 
@@ -1683,10 +1692,10 @@ PyMethodDef PyCoreMethods[] = {
 
 void InitializeModule()
 {
-	GPythonPropertyContainer.Reset(NewObject<UStruct>(GetTransientPackage(), TEXT("PythonProperties"), RF_Transient));
+	GetPythonPropertyContainerSingleton().Reset(NewObject<UStruct>(GetTransientPackage(), TEXT("PythonProperties"), RF_Transient));
 
-	GPythonTypeContainer.Reset(NewObject<UPackage>(nullptr, TEXT("/Engine/PythonTypes"), RF_Public | RF_Transient));
-	GPythonTypeContainer->SetPackageFlags(PKG_ContainsScript);
+	GetPythonTypeContainerSingleton().Reset(NewObject<UPackage>(nullptr, TEXT("/Engine/PythonTypes"), RF_Public | RF_Transient));
+	GetPythonTypeContainerSingleton()->SetPackageFlags(PKG_ContainsScript);
 
 	PyGenUtil::FNativePythonModule NativePythonModule;
 	NativePythonModule.PyModuleMethods = PyCoreMethods;
