Fixed typo when accessing property class type during Python glue validation

#rnx

diff --git a/Engine/Plugins/Experimental/PythonScriptPlugin/Source/PythonScriptPlugin/Private/PyWrapperTypeRegistry.cpp b/Engine/Plugins/Experimental/PythonScriptPlugin/Source/PythonScriptPlugin/Private/PyWrapperTypeRegistry.cpp
--- a/Engine/Plugins/Experimental/PythonScriptPlugin/Source/PythonScriptPlugin/Private/PyWrapperTypeRegistry.cpp
+++ b/Engine/Plugins/Experimental/PythonScriptPlugin/Source/PythonScriptPlugin/Private/PyWrapperTypeRegistry.cpp
@@ -736,9 +736,9 @@ PyTypeObject* FPyWrapperTypeRegistry::GenerateWrappedClassType(const UClass* InC
 		}
 		if (const UObjectPropertyBase* SelfPropObj = Cast<UObjectPropertyBase>(SelfParam.ParamProp))
 		{
-			if (SelfPropObj->GetClass()->IsChildOf(InFunc->GetOwnerClass()))
+			if (SelfPropObj->PropertyClass->IsChildOf(InFunc->GetOwnerClass()))
 			{
-				REPORT_PYTHON_GENERATION_ISSUE(Error, TEXT("Function '%s.%s' is marked as 'ScriptMethod' but the object argument type (%s) is a child of the the class type of the static function. This is not allowed."), *InFunc->GetOwnerClass()->GetName(), *InFunc->GetName(), *SelfPropObj->GetClass()->GetName());
+				REPORT_PYTHON_GENERATION_ISSUE(Error, TEXT("Function '%s.%s' is marked as 'ScriptMethod' but the object argument type (%s) is a child of the the class type of the static function. This is not allowed."), *InFunc->GetOwnerClass()->GetName(), *InFunc->GetName(), *SelfPropObj->PropertyClass->GetName());
 				return;
 			}
 		}
