Improved invalid parameter error when passing the incorrect struct type to Python

diff --git a/Engine/Plugins/Experimental/PythonScriptPlugin/Source/PythonScriptPlugin/Private/PyConversion.cpp b/Engine/Plugins/Experimental/PythonScriptPlugin/Source/PythonScriptPlugin/Private/PyConversion.cpp
--- a/Engine/Plugins/Experimental/PythonScriptPlugin/Source/PythonScriptPlugin/Private/PyConversion.cpp
+++ b/Engine/Plugins/Experimental/PythonScriptPlugin/Source/PythonScriptPlugin/Private/PyConversion.cpp
@@ -58,7 +58,7 @@ FPyConversionResult NativizeStructInstance(PyObject* PyObj, UScriptStruct* Struc
 		StructType->CopyScriptStruct(StructInstance, PyStruct->StructInstance);
 	}
 
-	PYCONVERSION_RETURN(Result, TEXT("Nativize"), *FString::Printf(TEXT("Cannot nativize '%s' as '%s'"), *PyUtil::GetFriendlyTypename(PyObj), *PyUtil::GetFriendlyTypename(PyStructType)));
+	PYCONVERSION_RETURN(Result, TEXT("NativizeStructInstance"), *FString::Printf(TEXT("Cannot nativize '%s' as '%s'"), *PyUtil::GetFriendlyTypename(PyObj), *PyUtil::GetFriendlyTypename(PyStructType)));
 }
 
 FPyConversionResult PythonizeStructInstance(UScriptStruct* StructType, const void* StructInstance, PyObject*& OutPyObj, const ESetErrorState SetErrorState)
@@ -639,7 +639,7 @@ FPyConversionResult NativizeEnumEntry(PyObject* PyObj, const UEnum* EnumType, in
 		OutVal = FPyWrapperEnum::GetEnumEntryValue(PyEnum);
 	}
 
-	PYCONVERSION_RETURN(Result, TEXT("Nativize"), *FString::Printf(TEXT("Cannot nativize '%s' as '%s'"), *PyUtil::GetFriendlyTypename(PyObj), *PyUtil::GetFriendlyTypename(PyEnumType)));
+	PYCONVERSION_RETURN(Result, TEXT("NativizeEnumEntry"), *FString::Printf(TEXT("Cannot nativize '%s' as '%s'"), *PyUtil::GetFriendlyTypename(PyObj), *PyUtil::GetFriendlyTypename(PyEnumType)));
 }
 
 FPyConversionResult PythonizeEnumEntry(const int64 Val, const UEnum* EnumType, PyObject*& OutPyObj, const ESetErrorState SetErrorState)
@@ -660,7 +660,7 @@ FPyConversionResult PythonizeEnumEntry(const int64 Val, const UEnum* EnumType, P
 		}
 	}
 
-	PYCONVERSION_RETURN(FPyConversionResult::Failure(), TEXT("Nativize"), *FString::Printf(TEXT("Cannot pythonize '%d' (int64) as '%s'"), Val, *PyUtil::GetFriendlyTypename(PyEnumType)));
+	PYCONVERSION_RETURN(FPyConversionResult::Failure(), TEXT("PythonizeEnumEntry"), *FString::Printf(TEXT("Cannot pythonize '%d' (int64) as '%s'"), Val, *PyUtil::GetFriendlyTypename(PyEnumType)));
 }
 
 PyObject* PythonizeEnumEntry(const int64 Val, const UEnum* EnumType, const ESetErrorState SetErrorState)
@@ -921,6 +921,14 @@ FPyConversionResult NativizeProperty_Direct(PyObject* PyObj, const UProperty* Pr
 				});
 			}
 		}
+		else
+		{
+			// Extra function scope as we don't want PYCONVERSION_RETURN to early return and skip the PYCONVERSION_PROPERTY_RETURN below
+			[&]()
+			{
+				PYCONVERSION_RETURN(Result, TEXT("NativizeStructInstance"), *FString::Printf(TEXT("Cannot nativize '%s' as '%s'"), *PyUtil::GetFriendlyTypename(PyObj), *PyUtil::GetFriendlyTypename(PyStructType)));
+			}();
+		}
 		PYCONVERSION_PROPERTY_RETURN(Result);
 	}
 
