MERGE: Fixed python UFunctions not showing up in the owning UClass' field list due to missing link list insertion

diff --git a/Engine/Plugins/Experimental/PythonScriptPlugin/Source/PythonScriptPlugin/Private/PyWrapperObject.cpp b/Engine/Plugins/Experimental/PythonScriptPlugin/Source/PythonScriptPlugin/Private/PyWrapperObject.cpp
--- a/Engine/Plugins/Experimental/PythonScriptPlugin/Source/PythonScriptPlugin/Private/PyWrapperObject.cpp
+++ b/Engine/Plugins/Experimental/PythonScriptPlugin/Source/PythonScriptPlugin/Private/PyWrapperObject.cpp
@@ -1357,6 +1357,13 @@ public:
 		Func->FunctionFlags |= (FUNC_Native | FUNC_Event | FUNC_BlueprintEvent | FUNC_BlueprintCallable);
 		FPyUFunctionDef::ApplyMetaData(InPyFuncDef, Func);
 		NewClass->AddFunctionToFunctionMap(Func, Func->GetFName());
+		// If the function is not in the linked list of class fields, insert it so that field iterators & funcs work
+		UField* FuncField = FindField<UField>(NewClass, Func->GetFName());
+		if (FuncField == nullptr)
+		{
+			Func->Next = NewClass->Children;
+			NewClass->Children = Func;
+		}
 		if (!Func->HasAnyFunctionFlags(FUNC_Static))
 		{
 			// Strip the zero'th 'self' argument when processing a non-static function
