MERGE

Added error handling for a missing self paramter in a python function that is exposed to the engine
Added missing class flags that allow a python object to expose a python function as a BP node without having to be present in the native base class
Update the Blueprint action database when a python class is loaded after the list has already been populated

diff --git a/Engine/Plugins/Experimental/PythonScriptPlugin/Source/PythonScriptPlugin/Private/PyWrapperObject.cpp b/Engine/Plugins/Experimental/PythonScriptPlugin/Source/PythonScriptPlugin/Private/PyWrapperObject.cpp
--- a/Engine/Plugins/Experimental/PythonScriptPlugin/Source/PythonScriptPlugin/Private/PyWrapperObject.cpp
+++ b/Engine/Plugins/Experimental/PythonScriptPlugin/Source/PythonScriptPlugin/Private/PyWrapperObject.cpp
@@ -17,6 +17,7 @@
 #include "UObject/PropertyPortFlags.h"
 #include "Engine/World.h"
 #include "Templates/Casts.h"
+#include "BlueprintActionDatabase.h"
 
 #if WITH_PYTHON
 
@@ -1102,8 +1103,8 @@ public:
 		// Create a new class with a temporary name; we will rename it as part of Finalize
 		const FString NewClassName = MakeUniqueObjectName(ClassOuter, UPythonGeneratedClass::StaticClass(), *FString::Printf(TEXT("%s_NEWINST"), *ClassName)).ToString();
 		NewClass = NewObject<UPythonGeneratedClass>(ClassOuter, *NewClassName, RF_Public | RF_Standalone | RF_Transient);
-		NewClass->SetMetaData(TEXT("BlueprintType"), TEXT("true"));
 		NewClass->SetSuperStruct(InSuperClass);
+		NewClass->ClassFlags |= CLASS_Native;
 	}
 
 	FPythonGeneratedClassBuilder(UPythonGeneratedClass* InOldClass, UClass* InSuperClass)
@@ -1117,8 +1118,8 @@ public:
 		// Create a new class with a temporary name; we will rename it as part of Finalize
 		const FString NewClassName = MakeUniqueObjectName(ClassOuter, UPythonGeneratedClass::StaticClass(), *FString::Printf(TEXT("%s_NEWINST"), *ClassName)).ToString();
 		NewClass = NewObject<UPythonGeneratedClass>(ClassOuter, *NewClassName, RF_Public | RF_Standalone | RF_Transient);
-		NewClass->SetMetaData(TEXT("BlueprintType"), TEXT("true"));
 		NewClass->SetSuperStruct(InSuperClass);
+		NewClass->ClassFlags |= CLASS_Native;
 	}
 
 	~FPythonGeneratedClassBuilder()
@@ -1181,6 +1182,12 @@ public:
 			UPythonGeneratedClass::ReparentDerivedClasses(OldClass, NewClass);
 		}	
 
+		if (FBlueprintActionDatabase* ActionDB = FBlueprintActionDatabase::TryGet())
+		{
+			// Notify Blueprints that there is a new class to add to the action list
+			ActionDB->RefreshClassActions(NewClass);
+		}
+
 		// Null the NewClass pointer so the destructor doesn't kill it
 		UPythonGeneratedClass* FinalizedClass = NewClass;
 		NewClass = nullptr;
@@ -1366,9 +1373,18 @@ public:
 		}
 		if (!Func->HasAnyFunctionFlags(FUNC_Static))
 		{
-			// Strip the zero'th 'self' argument when processing a non-static function
-			FuncArgNames.RemoveAt(0, 1, /*bAllowShrinking*/false);
-			FuncArgDefaults.RemoveAt(0, 1, /*bAllowShrinking*/false);
+			// Check for a malformed function rather than assert in the remove
+			if (FuncArgNames.Num() > 0 && FuncArgDefaults.Num() > 0)
+			{
+				// Strip the zero'th 'self' argument when processing a non-static function
+				FuncArgNames.RemoveAt(0, 1, /*bAllowShrinking*/false);
+				FuncArgDefaults.RemoveAt(0, 1, /*bAllowShrinking*/false);
+			}
+			else
+			{
+				PyUtil::SetPythonError(PyExc_Exception, PyType, *FString::Printf(TEXT("Incorrect number of arguments specified for '%s' (missing self?)"), *InFieldName));
+				return false;
+			}
 		}
 		if (!SuperFunc)
 		{
diff --git a/Engine/Plugins/Experimental/PythonScriptPlugin/Source/PythonScriptPlugin/Private/PyWrapperObject.h b/Engine/Plugins/Experimental/PythonScriptPlugin/Source/PythonScriptPlugin/Private/PyWrapperObject.h
--- a/Engine/Plugins/Experimental/PythonScriptPlugin/Source/PythonScriptPlugin/Private/PyWrapperObject.h
+++ b/Engine/Plugins/Experimental/PythonScriptPlugin/Source/PythonScriptPlugin/Private/PyWrapperObject.h
@@ -179,7 +179,7 @@ typedef TPyPtr<FPyWrapperObject> FPyWrapperObjectPtr;
 #endif	// WITH_PYTHON
 
 /** An Unreal class that was generated from a Python type */
-UCLASS()
+UCLASS(BlueprintType, Transient)
 class UPythonGeneratedClass : public UClass, public IPythonResourceOwner
 {
 	GENERATED_BODY()
diff --git a/Engine/Plugins/Experimental/PythonScriptPlugin/Source/PythonScriptPlugin/PythonScriptPlugin.Build.cs b/Engine/Plugins/Experimental/PythonScriptPlugin/Source/PythonScriptPlugin/PythonScriptPlugin.Build.cs
--- a/Engine/Plugins/Experimental/PythonScriptPlugin/Source/PythonScriptPlugin/PythonScriptPlugin.Build.cs
+++ b/Engine/Plugins/Experimental/PythonScriptPlugin/Source/PythonScriptPlugin/PythonScriptPlugin.Build.cs
@@ -38,6 +38,7 @@ namespace UnrealBuildTool.Rules
 						"LevelEditor",
 						"UnrealEd",
 						"EditorSubsystem",
+						"BlueprintGraph",
 					}
 				);
 			}
