MERGE: Removed the find field to determine if a function is already present in the python class' hierarchy

diff --git a/Engine/Plugins/Experimental/PythonScriptPlugin/Source/PythonScriptPlugin/Private/PyWrapperObject.cpp b/Engine/Plugins/Experimental/PythonScriptPlugin/Source/PythonScriptPlugin/Private/PyWrapperObject.cpp
--- a/Engine/Plugins/Experimental/PythonScriptPlugin/Source/PythonScriptPlugin/Private/PyWrapperObject.cpp
+++ b/Engine/Plugins/Experimental/PythonScriptPlugin/Source/PythonScriptPlugin/Private/PyWrapperObject.cpp
@@ -1340,6 +1340,9 @@ public:
 		if (!SuperFunc)
 		{
 			Func->FunctionFlags |= FUNC_Public;
+			// The function is not in the linked list of class fields, insert it so that field iterators & funcs work
+			Func->Next = NewClass->Children;
+			NewClass->Children = Func;
 		}
 		if (EnumHasAllFlags(InPyFuncDef->FuncFlags, EPyUFunctionDefFlags::Static))
 		{
@@ -1364,13 +1367,6 @@ public:
 		Func->FunctionFlags |= (FUNC_Native | FUNC_Event | FUNC_BlueprintEvent | FUNC_BlueprintCallable);
 		FPyUFunctionDef::ApplyMetaData(InPyFuncDef, Func);
 		NewClass->AddFunctionToFunctionMap(Func, Func->GetFName());
-		// If the function is not in the linked list of class fields, insert it so that field iterators & funcs work
-		UField* FuncField = FindField<UField>(NewClass, Func->GetFName());
-		if (FuncField == nullptr)
-		{
-			Func->Next = NewClass->Children;
-			NewClass->Children = Func;
-		}
 		if (!Func->HasAnyFunctionFlags(FUNC_Static))
 		{
 			// Check for a malformed function rather than assert in the remove
