MERGE: Fixed Python defined ufunctions having the parameters reversed due to how properties are added to functions

diff --git a/Engine/Plugins/Experimental/PythonScriptPlugin/Source/PythonScriptPlugin/Private/PyWrapperObject.cpp b/Engine/Plugins/Experimental/PythonScriptPlugin/Source/PythonScriptPlugin/Private/PyWrapperObject.cpp
--- a/Engine/Plugins/Experimental/PythonScriptPlugin/Source/PythonScriptPlugin/Private/PyWrapperObject.cpp
+++ b/Engine/Plugins/Experimental/PythonScriptPlugin/Source/PythonScriptPlugin/Private/PyWrapperObject.cpp
@@ -1426,7 +1426,9 @@ public:
 					}
 				}
 			}
-			for (int32 ArgIndex = 0; ArgIndex < FuncArgNames.Num(); ++ArgIndex)
+			// Adding properties to a function inserts them into a linked list, so we loop backwards to get the order right
+			int32 ArgIndex = FuncArgNames.Num() - 1;
+			while (ArgIndex >= 0)
 			{
 				PyObject* ArgTypeObj = PySequence_GetItem(InPyFuncDef->FuncParamTypes, ArgIndex);
 				UProperty* ArgProp = PyUtil::CreateProperty(ArgTypeObj, 1, Func, *FuncArgNames[ArgIndex]);
@@ -1437,6 +1439,7 @@ public:
 				}
 				ArgProp->PropertyFlags |= CPF_Parm;
 				Func->AddCppProperty(ArgProp);
+				ArgIndex--;
 			}
 		}
 		// Apply the defaults to the function arguments and build the Python method params
