MERGE: Fixed the order that the return value is pushed into the list of parameters for a UFunction that maps to a Python function without a native super function backing it
- This fixes out params being at the wrong offsets causing the copy of the data to look at the wrong locations

diff --git a/Engine/Plugins/Experimental/PythonScriptPlugin/Source/PythonScriptPlugin/Private/PyWrapperObject.cpp b/Engine/Plugins/Experimental/PythonScriptPlugin/Source/PythonScriptPlugin/Private/PyWrapperObject.cpp
--- a/Engine/Plugins/Experimental/PythonScriptPlugin/Source/PythonScriptPlugin/Private/PyWrapperObject.cpp
+++ b/Engine/Plugins/Experimental/PythonScriptPlugin/Source/PythonScriptPlugin/Private/PyWrapperObject.cpp
@@ -1392,6 +1392,24 @@ public:
 				return false;
 			}
 
+			{			
+				// Adding properties to a function inserts them into a linked list, so we loop backwards to get the order right
+				int32 ArgIndex = FuncArgNames.Num() - 1;
+				while (ArgIndex >= 0)
+				{
+					PyObject* ArgTypeObj = PySequence_GetItem(InPyFuncDef->FuncParamTypes, ArgIndex);
+					UProperty* ArgProp = PyUtil::CreateProperty(ArgTypeObj, 1, Func, *FuncArgNames[ArgIndex]);
+					if (!ArgProp)
+					{
+						PyUtil::SetPythonError(PyExc_Exception, PyType, *FString::Printf(TEXT("Failed to create property (%s) for function '%s' argument '%s'"), *PyUtil::GetFriendlyTypename(ArgTypeObj), *InFieldName, *FuncArgNames[ArgIndex]));
+						return false;
+					}
+					ArgProp->PropertyFlags |= CPF_Parm;
+					Func->AddCppProperty(ArgProp);
+					ArgIndex--;
+				}
+			}
+
 			// Build the arguments struct if not overriding a function
 			if (InPyFuncDef->FuncRetType && InPyFuncDef->FuncRetType != Py_None)
 			{
@@ -1426,21 +1444,6 @@ public:
 					}
 				}
 			}
-			// Adding properties to a function inserts them into a linked list, so we loop backwards to get the order right
-			int32 ArgIndex = FuncArgNames.Num() - 1;
-			while (ArgIndex >= 0)
-			{
-				PyObject* ArgTypeObj = PySequence_GetItem(InPyFuncDef->FuncParamTypes, ArgIndex);
-				UProperty* ArgProp = PyUtil::CreateProperty(ArgTypeObj, 1, Func, *FuncArgNames[ArgIndex]);
-				if (!ArgProp)
-				{
-					PyUtil::SetPythonError(PyExc_Exception, PyType, *FString::Printf(TEXT("Failed to create property (%s) for function '%s' argument '%s'"), *PyUtil::GetFriendlyTypename(ArgTypeObj), *InFieldName, *FuncArgNames[ArgIndex]));
-					return false;
-				}
-				ArgProp->PropertyFlags |= CPF_Parm;
-				Func->AddCppProperty(ArgProp);
-				ArgIndex--;
-			}
 		}
 		// Apply the defaults to the function arguments and build the Python method params
 		PyGenUtil::FGeneratedWrappedFunction GeneratedWrappedFunction;
