IsFunctionImplementedInBlueprint for python generated classes

diff --git a/Engine/Plugins/Experimental/PythonScriptPlugin/Source/PythonScriptPlugin/Private/PyWrapperObject.cpp b/Engine/Plugins/Experimental/PythonScriptPlugin/Source/PythonScriptPlugin/Private/PyWrapperObject.cpp
--- a/Engine/Plugins/Experimental/PythonScriptPlugin/Source/PythonScriptPlugin/Private/PyWrapperObject.cpp
+++ b/Engine/Plugins/Experimental/PythonScriptPlugin/Source/PythonScriptPlugin/Private/PyWrapperObject.cpp
@@ -1687,6 +1687,12 @@ void UPythonGeneratedClass::PostInitInstance(UObject* InObj)
 	}
 }
 
+bool UPythonGeneratedClass::IsFunctionImplementedInBlueprint(FName InFunctionName) const
+{
+	UFunction* Function = FindFunctionByName(InFunctionName);
+	return Function && Function->GetOuter() && Function->GetOuter()->IsA(UPythonGeneratedClass::StaticClass());
+}
+
 UPythonGeneratedClass* UPythonGeneratedClass::GenerateClass(PyTypeObject* InPyType)
 {
 	// Get the correct super class from the parent type in Python
diff --git a/Engine/Plugins/Experimental/PythonScriptPlugin/Source/PythonScriptPlugin/Private/PyWrapperObject.h b/Engine/Plugins/Experimental/PythonScriptPlugin/Source/PythonScriptPlugin/Private/PyWrapperObject.h
--- a/Engine/Plugins/Experimental/PythonScriptPlugin/Source/PythonScriptPlugin/Private/PyWrapperObject.h
+++ b/Engine/Plugins/Experimental/PythonScriptPlugin/Source/PythonScriptPlugin/Private/PyWrapperObject.h
@@ -193,6 +193,8 @@ public:
 	//~ UClass interface
 	virtual void PostInitInstance(UObject* InObj) override;
 
+	virtual bool IsFunctionImplementedInBlueprint(FName InFunctionName) const override;
+
 	/** Generate an Unreal class from the given Python type */
 	static UPythonGeneratedClass* GenerateClass(PyTypeObject* InPyType);
 
