Updated Python error message about ScriptMethod functions to mention UPRAM(ref)

diff --git a/Engine/Plugins/Experimental/PythonScriptPlugin/Source/PythonScriptPlugin/Private/PyWrapperTypeRegistry.cpp b/Engine/Plugins/Experimental/PythonScriptPlugin/Source/PythonScriptPlugin/Private/PyWrapperTypeRegistry.cpp
--- a/Engine/Plugins/Experimental/PythonScriptPlugin/Source/PythonScriptPlugin/Private/PyWrapperTypeRegistry.cpp
+++ b/Engine/Plugins/Experimental/PythonScriptPlugin/Source/PythonScriptPlugin/Private/PyWrapperTypeRegistry.cpp
@@ -699,15 +699,29 @@ PyTypeObject* FPyWrapperTypeRegistry::GenerateWrappedClassType(const UClass* InC
 			return;
 		}
 
+		static auto IsStructOrObjectProperty = [](const UProperty* InProp) { return InProp && (InProp->IsA<UStructProperty>() || InProp->IsA<UObjectPropertyBase>()); };
+
 		// Get the type to hoist this method to (this should be the first parameter)
 		PyGenUtil::FGeneratedWrappedMethodParameter SelfParam;
-		if (InTypeMethod.MethodFunc.InputParams.Num() > 0 && (InTypeMethod.MethodFunc.InputParams[0].ParamProp->IsA<UStructProperty>() || InTypeMethod.MethodFunc.InputParams[0].ParamProp->IsA<UObjectPropertyBase>()))
+		if (InTypeMethod.MethodFunc.InputParams.Num() > 0 && IsStructOrObjectProperty(InTypeMethod.MethodFunc.InputParams[0].ParamProp))
 		{
 			SelfParam = InTypeMethod.MethodFunc.InputParams[0];
 		}
 		if (!SelfParam.ParamProp)
 		{
-			REPORT_PYTHON_GENERATION_ISSUE(Error, TEXT("Function '%s.%s' is marked as 'ScriptMethod' but doesn't contain a valid struct or object as its first argument."), *InFunc->GetOwnerClass()->GetName(), *InFunc->GetName());
+			// Hint that UPRAM(ref) may be missing on first parameter
+			const UProperty* PropertyPossiblyMissingMacro = nullptr;
+			for (TFieldIterator<const UProperty> ParamIt(InFunc); ParamIt; ++ParamIt)
+			{
+				const UProperty* Param = *ParamIt;
+				if (PyUtil::IsOutputParameter(Param) && !PyUtil::IsInputParameter(Param) && !Param->HasAnyPropertyFlags(CPF_ReturnParm) && IsStructOrObjectProperty(Param))
+				{
+					PropertyPossiblyMissingMacro = Param;
+				}
+				break;
+			}
+
+			REPORT_PYTHON_GENERATION_ISSUE(Error, TEXT("Function '%s.%s' is marked as 'ScriptMethod' but doesn't contain a valid struct or object as its first argument.%s"), *InFunc->GetOwnerClass()->GetName(), *InFunc->GetName(), PropertyPossiblyMissingMacro ? TEXT(" UPARAM(ref) may be missing on the first argument.") : TEXT(""));
 			return;
 		}
 		if (const UObjectPropertyBase* SelfPropObj = Cast<UObjectPropertyBase>(SelfParam.ParamProp))
