New UI menu item added to clear recent Python files menu

diff --git a/Engine/Plugins/Experimental/PythonScriptPlugin/Source/PythonScriptPlugin/Private/PythonScriptPlugin.cpp b/Engine/Plugins/Experimental/PythonScriptPlugin/Source/PythonScriptPlugin/Private/PythonScriptPlugin.cpp
--- a/Engine/Plugins/Experimental/PythonScriptPlugin/Source/PythonScriptPlugin/Private/PythonScriptPlugin.cpp
+++ b/Engine/Plugins/Experimental/PythonScriptPlugin/Source/PythonScriptPlugin/Private/PythonScriptPlugin.cpp
@@ -317,6 +317,7 @@ private:
 
 	void MakeRecentPythonScriptMenu(FMenuBuilder& InMenuBuilder)
 	{
+		InMenuBuilder.BeginSection("Files");
 		for (int32 Index = RecentsFiles.Num() - 1; Index >= 0; --Index)
 		{
 			InMenuBuilder.AddMenuEntry(
@@ -326,7 +327,16 @@ private:
 				FUIAction(FExecuteAction::CreateRaw(this, &FPythonCommandMenuImpl::Menu_ExecutePythonRecent, Index))
 			);
 		}
-		int32 Variable1 = 1;
+		InMenuBuilder.EndSection();
+
+		InMenuBuilder.BeginSection("Clear");
+		InMenuBuilder.AddMenuEntry(
+			LOCTEXT("ClearRecentPython", "Clear Recent Python Scripts"),
+			FText::GetEmpty(),
+			FSlateIcon(),
+			FUIAction(FExecuteAction::CreateRaw(this, &FPythonCommandMenuImpl::Menu_ClearRecentPython))
+		);
+		InMenuBuilder.EndSection();
 	}
 
 	void CreateMenu(FMenuBuilder& MenuBuilder)
@@ -357,6 +367,15 @@ private:
 		}
 	}
 
+	void Menu_ClearRecentPython()
+	{
+		if (RecentsFiles.Num() > 0)
+		{
+			RecentsFiles.Reset();
+			bRecentsFilesDirty = true;
+		}
+	}
+
 	void Menu_ExecutePython()
 	{
 		TArray<FString> OpenedFiles;
