Python file menu conversion
#rnx

diff --git a/Engine/Plugins/Experimental/PythonScriptPlugin/Source/PythonScriptPlugin/Private/PythonScriptPlugin.cpp b/Engine/Plugins/Experimental/PythonScriptPlugin/Source/PythonScriptPlugin/Private/PythonScriptPlugin.cpp
--- a/Engine/Plugins/Experimental/PythonScriptPlugin/Source/PythonScriptPlugin/Private/PythonScriptPlugin.cpp
+++ b/Engine/Plugins/Experimental/PythonScriptPlugin/Source/PythonScriptPlugin/Private/PythonScriptPlugin.cpp
@@ -32,7 +32,7 @@
 #include "Engine/Engine.h"
 #include "Framework/Application/SlateApplication.h"
 #include "Framework/MultiBox/MultiBoxBuilder.h"
-#include "LevelEditor.h"
+#include "EditorMenuSubsystem.h"
 #include "Misc/ConfigCacheIni.h"
 #include "Misc/MessageDialog.h"
 #endif	// WITH_EDITOR
@@ -219,20 +219,12 @@ public:
 	{
 		LoadConfig();
 
-		// Create & Add menu extension
-		MenuExtender = MakeShareable(new FExtender);
-		MenuExtender->AddMenuExtension("FileLoadAndSave", EExtensionHook::After, nullptr, FMenuExtensionDelegate::CreateRaw(this, &FPythonCommandMenuImpl::CreateMenu));
-
-		FLevelEditorModule& LevelEditorModule = FModuleManager::LoadModuleChecked<FLevelEditorModule>("LevelEditor");
-		LevelEditorModule.GetMenuExtensibilityManager()->AddExtender(MenuExtender);
+		RegisterMenus();
 	}
 
 	virtual void OnShutdownMenu() override
 	{
-		// Remove menu extension
-		FLevelEditorModule& LevelEditorModule = FModuleManager::LoadModuleChecked<FLevelEditorModule>("LevelEditor");
-		LevelEditorModule.GetMenuExtensibilityManager()->RemoveExtender(MenuExtender);
-		MenuExtender = nullptr;
+		UEditorMenuSubsystem::UnregisterOwner(this);
 
 		// Write to file
 		if (bRecentsFilesDirty)
@@ -321,47 +313,52 @@ private:
 		GConfig->Flush(false);
 	}
 
-	void MakeRecentPythonScriptMenu(FMenuBuilder& InMenuBuilder)
+	void MakeRecentPythonScriptMenu(UEditorMenu* InMenu)
 	{
-		InMenuBuilder.BeginSection("Files");
+		FEditorMenuOwnerScoped OwnerScoped(this);
+		FEditorMenuSection& Section = InMenu->AddSection("Files");
 		for (int32 Index = RecentsFiles.Num() - 1; Index >= 0; --Index)
 		{
-			InMenuBuilder.AddMenuEntry(
+			Section.AddMenuEntry(
+				NAME_None,
 				FText::FromString(RecentsFiles[Index]),
 				FText::GetEmpty(),
 				FSlateIcon(),
 				FUIAction(FExecuteAction::CreateRaw(this, &FPythonCommandMenuImpl::Menu_ExecutePythonRecent, Index))
 			);
 		}
-		InMenuBuilder.EndSection();
 
-		InMenuBuilder.BeginSection("Clear");
-		InMenuBuilder.AddMenuEntry(
+		FEditorMenuSection& ClearSection = InMenu->AddSection("Clear");
+		Section.AddMenuEntry(
+			"ClearRecentPython",
 			LOCTEXT("ClearRecentPython", "Clear Recent Python Scripts"),
 			FText::GetEmpty(),
 			FSlateIcon(),
 			FUIAction(FExecuteAction::CreateRaw(this, &FPythonCommandMenuImpl::Menu_ClearRecentPython))
 		);
-		InMenuBuilder.EndSection();
 	}
 
-	void CreateMenu(FMenuBuilder& MenuBuilder)
+	void RegisterMenus()
 	{
-		MenuBuilder.BeginSection("Python", LOCTEXT("Python", "Python"));
-		MenuBuilder.AddMenuEntry(
+		FEditorMenuOwnerScoped OwnerScoped(this);
+		UEditorMenu* Menu = UEditorMenuSubsystem::Get()->ExtendMenu("LevelEditor.MainMenu.File");
+		FEditorMenuSection& Section = Menu->AddSection("Python", LOCTEXT("Python", "Python"), FEditorMenuInsert("FileLoadAndSave", EEditorMenuInsertType::After));
+		Section.AddMenuEntry(
+			"OpenPython",
 			LOCTEXT("OpenPython", "Execute Python Script"),
 			LOCTEXT("OpenPythonTooltip", "Open a Python Script file and Execute it."),
 			FSlateIcon(),
 			FUIAction(FExecuteAction::CreateRaw(this, &FPythonCommandMenuImpl::Menu_ExecutePython))
 		);
-		MenuBuilder.AddSubMenu(
+		Section.AddEntry(FEditorMenuEntry::InitSubMenu(
+			"LevelEditor.MainMenu.File",
+			"RecentPythonsSubMenu",
 			LOCTEXT("RecentPythonsSubMenu", "Recent Python Scripts"),
 			LOCTEXT("RecentPythonsSubMenu_ToolTip", "Select a recent Python Script file and Execute it."),
-			FNewMenuDelegate::CreateRaw(this, &FPythonCommandMenuImpl::MakeRecentPythonScriptMenu),
+			FNewEditorMenuDelegate::CreateRaw(this, &FPythonCommandMenuImpl::MakeRecentPythonScriptMenu),
 			false,
 			FSlateIcon(FEditorStyle::GetStyleSetName(), "MainFrame.RecentLevels")
-		);
-		MenuBuilder.EndSection();
+		));
 	}
 
 	void Menu_ExecutePythonRecent(int32 Index)
@@ -414,7 +411,6 @@ private:
 	}
 
 private:
-	TSharedPtr<FExtender> MenuExtender;
 
 	TArray<FString> RecentsFiles;
 	FString LastDirectory;
@@ -503,15 +499,29 @@ void FPythonScriptPlugin::StartupModule()
 	RemoteExecution = MakeUnique<FPythonScriptRemoteExecution>(this);
 
 #if WITH_EDITOR
-	check(CmdMenu == nullptr);
-	CmdMenu = new FPythonCommandMenuImpl();
-	CmdMenu->OnStartupMenu();
+	FCoreDelegates::OnPostEngineInit.AddRaw(this, &FPythonScriptPlugin::OnPostEngineInit);
 #endif // WITH_EDITOR
 
 	FCoreDelegates::OnPreExit.AddRaw(this, &FPythonScriptPlugin::ShutdownPython);
 #endif	// WITH_PYTHON
 }
 
+#if WITH_PYTHON
+#if WITH_EDITOR
+void FPythonScriptPlugin::OnPostEngineInit()
+{
+	check(CmdMenu == nullptr);
+	CmdMenu = new FPythonCommandMenuImpl();
+	CmdMenu->OnStartupMenu();
+
+	UEditorMenuSubsystem::Get()->RegisterStringCommandHandler("Python", FEditorMenuExecuteString::CreateLambda([this](const FString& InString, const FEditorMenuContext& InContext) {
+		ExecPythonCommand(*InString);
+	}));
+}
+
+#endif // WITH_EDITOR
+#endif // WITH_PYTHON
+
 void FPythonScriptPlugin::ShutdownModule()
 {
 #if WITH_PYTHON
@@ -520,10 +530,17 @@ void FPythonScriptPlugin::ShutdownModule()
 	RemoteExecution.Reset();
 
 #if WITH_EDITOR
+	FCoreDelegates::OnPostEngineInit.RemoveAll(this);
+
 	check(CmdMenu);
 	CmdMenu->OnShutdownMenu();
 	delete CmdMenu;
 	CmdMenu = nullptr;
+
+	if (UEditorMenuSubsystem* EditorMenus = UEditorMenuSubsystem::TryGet())
+	{
+		EditorMenus->UnregisterStringCommandHandler("Python");
+	}
 #endif // WITH_EDITOR
 
 	IModularFeatures::Get().UnregisterModularFeature(IConsoleCommandExecutor::ModularFeatureName(), &CmdExec);
diff --git a/Engine/Plugins/Experimental/PythonScriptPlugin/Source/PythonScriptPlugin/Private/PythonScriptPlugin.h b/Engine/Plugins/Experimental/PythonScriptPlugin/Source/PythonScriptPlugin/Private/PythonScriptPlugin.h
--- a/Engine/Plugins/Experimental/PythonScriptPlugin/Source/PythonScriptPlugin/Private/PythonScriptPlugin.h
+++ b/Engine/Plugins/Experimental/PythonScriptPlugin/Source/PythonScriptPlugin/Private/PythonScriptPlugin.h
@@ -107,6 +107,11 @@ public:
 	virtual bool Exec(UWorld* InWorld, const TCHAR* Cmd, FOutputDevice& Ar) override;
 
 #if WITH_PYTHON
+
+#if WITH_EDITOR	
+	void OnPostEngineInit();
+#endif
+
 	/** Sync the remote execution environment to the current settings, starting or stopping it as required */
 	void SyncRemoteExecutionToSettings();
 
diff --git a/Engine/Plugins/Experimental/PythonScriptPlugin/Source/PythonScriptPlugin/PythonScriptPlugin.Build.cs b/Engine/Plugins/Experimental/PythonScriptPlugin/Source/PythonScriptPlugin/PythonScriptPlugin.Build.cs
--- a/Engine/Plugins/Experimental/PythonScriptPlugin/Source/PythonScriptPlugin/PythonScriptPlugin.Build.cs
+++ b/Engine/Plugins/Experimental/PythonScriptPlugin/Source/PythonScriptPlugin/PythonScriptPlugin.Build.cs
@@ -31,6 +31,12 @@ namespace UnrealBuildTool.Rules
 
 			if (Target.bBuildEditor == true)
 			{
+				PublicDependencyModuleNames.AddRange(
+					new string[] {
+						"EditorMenus"
+					}
+				);
+
 				PrivateDependencyModuleNames.AddRange(
 					new string[] {
 						"DesktopPlatform",
