Fix crash when cooking and running editor executable in non-editor modes
#rnx

diff --git a/Engine/Plugins/Experimental/PythonScriptPlugin/Source/PythonScriptPlugin/Private/PythonScriptPlugin.cpp b/Engine/Plugins/Experimental/PythonScriptPlugin/Source/PythonScriptPlugin/Private/PythonScriptPlugin.cpp
--- a/Engine/Plugins/Experimental/PythonScriptPlugin/Source/PythonScriptPlugin/Private/PythonScriptPlugin.cpp
+++ b/Engine/Plugins/Experimental/PythonScriptPlugin/Source/PythonScriptPlugin/Private/PythonScriptPlugin.cpp
@@ -510,13 +510,16 @@ void FPythonScriptPlugin::StartupModule()
 #if WITH_EDITOR
 void FPythonScriptPlugin::OnPostEngineInit()
 {
-	check(CmdMenu == nullptr);
-	CmdMenu = new FPythonCommandMenuImpl();
-	CmdMenu->OnStartupMenu();
+	if (UEditorMenuSubsystem::IsRunningEditorUI())
+	{
+		check(CmdMenu == nullptr);
+		CmdMenu = new FPythonCommandMenuImpl();
+		CmdMenu->OnStartupMenu();
 
-	UEditorMenuSubsystem::Get()->RegisterStringCommandHandler("Python", FEditorMenuExecuteString::CreateLambda([this](const FString& InString, const FEditorMenuContext& InContext) {
-		ExecPythonCommand(*InString);
-	}));
+		UEditorMenuSubsystem::Get()->RegisterStringCommandHandler("Python", FEditorMenuExecuteString::CreateLambda([this](const FString& InString, const FEditorMenuContext& InContext) {
+			ExecPythonCommand(*InString);
+		}));
+	}
 }
 
 #endif // WITH_EDITOR
@@ -532,10 +535,12 @@ void FPythonScriptPlugin::ShutdownModule()
 #if WITH_EDITOR
 	FCoreDelegates::OnPostEngineInit.RemoveAll(this);
 
-	check(CmdMenu);
-	CmdMenu->OnShutdownMenu();
-	delete CmdMenu;
-	CmdMenu = nullptr;
+	if (CmdMenu)
+	{
+		CmdMenu->OnShutdownMenu();
+		delete CmdMenu;
+		CmdMenu = nullptr;
+	}
 
 	if (UEditorMenuSubsystem* EditorMenus = UEditorMenuSubsystem::TryGet())
 	{
diff --git a/Engine/Source/Editor/EditorMenus/Private/EditorMenuSubsystem.cpp b/Engine/Source/Editor/EditorMenus/Private/EditorMenuSubsystem.cpp
--- a/Engine/Source/Editor/EditorMenus/Private/EditorMenuSubsystem.cpp
+++ b/Engine/Source/Editor/EditorMenus/Private/EditorMenuSubsystem.cpp
@@ -73,6 +73,16 @@ UEditorMenuSubsystem::UEditorMenuSubsystem() :
 {
 }
 
+bool UEditorMenuSubsystem::IsRunningEditorUI()
+{
+	if (!FSlateApplication::IsInitialized())
+	{
+		return false;
+	}
+
+	return !IsRunningCommandlet() && !IsRunningGame() && !IsRunningDedicatedServer() && !IsRunningClientOnly();
+}
+
 void UEditorMenuSubsystem::Initialize(FSubsystemCollectionBase& Collection)
 {
 }
diff --git a/Engine/Source/Editor/EditorMenus/Public/EditorMenuSubsystem.h b/Engine/Source/Editor/EditorMenus/Public/EditorMenuSubsystem.h
--- a/Engine/Source/Editor/EditorMenus/Public/EditorMenuSubsystem.h
+++ b/Engine/Source/Editor/EditorMenus/Public/EditorMenuSubsystem.h
@@ -117,6 +117,14 @@ public:
 		}
 	}
 
+	/**
+	 * Returns true if slate initialized and editor GUI is being used.
+	 * The application should have been initialized before this method is called.
+	 *
+	 * @return	True if slate initialized and editor GUI is being used.
+	 */
+	static bool IsRunningEditorUI();
+
 	// UEditorSubsystem implementation
 	virtual void Initialize(FSubsystemCollectionBase& Collection);
 	virtual void Deinitialize();
