Remove invalid warning when reinstancing Python generated enum types

diff --git a/Engine/Plugins/Experimental/PythonScriptPlugin/Source/PythonScriptPlugin/Private/PyWrapperEnum.cpp b/Engine/Plugins/Experimental/PythonScriptPlugin/Source/PythonScriptPlugin/Private/PyWrapperEnum.cpp
--- a/Engine/Plugins/Experimental/PythonScriptPlugin/Source/PythonScriptPlugin/Private/PyWrapperEnum.cpp
+++ b/Engine/Plugins/Experimental/PythonScriptPlugin/Source/PythonScriptPlugin/Private/PyWrapperEnum.cpp
@@ -721,12 +721,17 @@ public:
 		: EnumName(InEnumName)
 		, PyType(InPyType)
 		, NewEnum(nullptr)
+		, bDidExist(false)
 	{
 		UObject* EnumOuter = GetPythonTypeContainer();
 
 		// Enum instances are re-used if they already exist
 		NewEnum = FindObject<UPythonGeneratedEnum>(EnumOuter, *EnumName);
-		if (!NewEnum)
+		if (NewEnum)
+		{
+			bDidExist = true;
+		}
+		else
 		{
 			NewEnum = NewObject<UPythonGeneratedEnum>(EnumOuter, *EnumName, RF_Public | RF_Standalone | RF_Transient);
 			NewEnum->SetMetaData(TEXT("BlueprintType"), TEXT("true"));
@@ -767,7 +772,7 @@ public:
 
 		// Map the Unreal enum to the Python type
 		NewEnum->PyType = FPyTypeObjectPtr::NewReference(PyType);
-		FPyWrapperTypeRegistry::Get().RegisterWrappedEnumType(NewEnum->GetFName(), PyType);
+		FPyWrapperTypeRegistry::Get().RegisterWrappedEnumType(NewEnum->GetFName(), PyType, !bDidExist);
 
 		// Null the NewEnum pointer so the destructor doesn't kill it
 		UPythonGeneratedEnum* FinalizedEnum = NewEnum;
@@ -837,6 +842,7 @@ private:
 	FString EnumName;
 	PyTypeObject* PyType;
 	UPythonGeneratedEnum* NewEnum;
+	bool bDidExist;
 };
 
 void UPythonGeneratedEnum::ReleasePythonResources()
