Register enum entries from Python in a deterministic order using their value

diff --git a/Engine/Plugins/Experimental/PythonScriptPlugin/Source/PythonScriptPlugin/Private/PyWrapperEnum.cpp b/Engine/Plugins/Experimental/PythonScriptPlugin/Source/PythonScriptPlugin/Private/PyWrapperEnum.cpp
--- a/Engine/Plugins/Experimental/PythonScriptPlugin/Source/PythonScriptPlugin/Private/PyWrapperEnum.cpp
+++ b/Engine/Plugins/Experimental/PythonScriptPlugin/Source/PythonScriptPlugin/Private/PyWrapperEnum.cpp
@@ -791,6 +791,7 @@ public:
 
 		// Build the definition data for the new enum value
 		UPythonGeneratedEnum::FEnumValueDef& EnumValueDef = *NewEnum->EnumValueDefs.Add_GetRef(MakeShared<UPythonGeneratedEnum::FEnumValueDef>());
+		EnumValueDef.PyIndex = NewEnum->EnumValueDefs.Num() - 1;
 		EnumValueDef.Value = EnumValue;
 		EnumValueDef.Name = InFieldName;
 
@@ -800,6 +801,12 @@ public:
 private:
 	bool RegisterDescriptors(const TArray<FPyUValueDef*>& InPyValueDefs)
 	{
+		// The enum entries came from a Python dict, so sort them into a consistent order (using the value) before registering them
+		NewEnum->EnumValueDefs.Sort([](const TSharedPtr<UPythonGeneratedEnum::FEnumValueDef>& EnumValueDefOne, const TSharedPtr<UPythonGeneratedEnum::FEnumValueDef>& EnumValueDefTwo)
+		{
+			return EnumValueDefOne->Value < EnumValueDefTwo->Value;
+		});
+
 		// Populate the enum with its values
 		check(InPyValueDefs.Num() == NewEnum->EnumValueDefs.Num());
 		{
@@ -816,13 +823,15 @@ private:
 			}
 
 			// Can't set the meta-data until SetEnums has been called
-			for (int32 EnumEntryIndex = 0; EnumEntryIndex < InPyValueDefs.Num(); ++EnumEntryIndex)
+			// Note: Beware, InPyValueDefs is not guaranteed to be in the same order due to the sort above - index it via the PyIndex of the value definition!
+			for (int32 EnumValueIndex = 0; EnumValueIndex < NewEnum->EnumValueDefs.Num(); ++EnumValueIndex)
 			{
-				FPyUValueDef::ApplyMetaData(InPyValueDefs[EnumEntryIndex], [this, EnumEntryIndex](const FString& InMetaDataKey, const FString& InMetaDataValue)
+				TSharedPtr<UPythonGeneratedEnum::FEnumValueDef>& EnumValueDef = NewEnum->EnumValueDefs[EnumValueIndex];
+				FPyUValueDef::ApplyMetaData(InPyValueDefs[EnumValueDef->PyIndex], [this, EnumValueIndex](const FString& InMetaDataKey, const FString& InMetaDataValue)
 				{
-					NewEnum->SetMetaData(*InMetaDataKey, *InMetaDataValue, EnumEntryIndex);
+					NewEnum->SetMetaData(*InMetaDataKey, *InMetaDataValue, EnumValueIndex);
 				});
-				NewEnum->EnumValueDefs[EnumEntryIndex]->DocString = PyGenUtil::GetEnumEntryTooltip(NewEnum, EnumEntryIndex);
+				NewEnum->EnumValueDefs[EnumValueIndex]->DocString = PyGenUtil::GetEnumEntryTooltip(NewEnum, EnumValueIndex);
 			}
 		}
 
diff --git a/Engine/Plugins/Experimental/PythonScriptPlugin/Source/PythonScriptPlugin/Private/PyWrapperEnum.h b/Engine/Plugins/Experimental/PythonScriptPlugin/Source/PythonScriptPlugin/Private/PyWrapperEnum.h
--- a/Engine/Plugins/Experimental/PythonScriptPlugin/Source/PythonScriptPlugin/Private/PyWrapperEnum.h
+++ b/Engine/Plugins/Experimental/PythonScriptPlugin/Source/PythonScriptPlugin/Private/PyWrapperEnum.h
@@ -127,8 +127,11 @@ private:
 	/** Definition data for an Unreal enum value generated from a Python type */
 	struct FEnumValueDef
 	{
+		/** Index of the enum value as it was registered with us from the Python */
+		int32 PyIndex = 0;
+
 		/** Numeric value of the enum value */
-		int64 Value;
+		int64 Value = 0;
 
 		/** Name of the enum value */
 		FString Name;
