Exposed UserDefinedEnum types to Python

This is the first phase of allowing Blueprint types to be generated for Python, which must be accessed via unreal.get_blueprint_generated_types(...) as asset names are not guaranteed to be unique, so they're not added to the unreal module.

Once a type has been generated for Python, it will be updated if the underlying asset is changed. Should the asset be deleted then the Python type will be reset to clean state and become mostly unusable.

This change also ensures that the Unreal type pointers referenced by the Python meta-data are ARO'd correctly. This will cause a warning if you try and delete an asset with a type used by Python.

 MyEnum = unreal.get_blueprint_generated_types('/Game/MyEnum') # This can also accept multiple arguments or a list, and will return a tuple of types
 print(MyEnum.MY_ENTRY)

diff --git a/Engine/Plugins/Experimental/PythonScriptPlugin/Source/PythonScriptPlugin/Private/PyEngine.cpp b/Engine/Plugins/Experimental/PythonScriptPlugin/Source/PythonScriptPlugin/Private/PyEngine.cpp
--- a/Engine/Plugins/Experimental/PythonScriptPlugin/Source/PythonScriptPlugin/Private/PyEngine.cpp
+++ b/Engine/Plugins/Experimental/PythonScriptPlugin/Source/PythonScriptPlugin/Private/PyEngine.cpp
@@ -93,7 +93,119 @@ PyTypeObject PySelectedActorIteratorType = InitializePyActorIteratorType<FSelect
 namespace PyEngine
 {
 
+PyObject* GetBlueprintGeneratedTypes(PyObject* InSelf, PyObject* InArgs)
+{
+	TArray<FString> AssetsToGenerate;
+
+	// Extract the assets to generate
+	if (InArgs)
+	{
+		const Py_ssize_t ArgsLen = PyTuple_Size(InArgs);
+		for (Py_ssize_t ArgIndex = 0; ArgIndex < ArgsLen; ++ArgIndex)
+		{
+			PyObject* PyArg = PyTuple_GetItem(InArgs, ArgIndex);
+			if (PyArg)
+			{
+				// Is this some kind of container, or a single value?
+#if PY_MAJOR_VERSION < 3
+				const bool bIsStringType = PyUnicode_Check(PyArg) || PyString_Check(PyArg);
+#else	// PY_MAJOR_VERSION < 3
+				const bool bIsStringType = PyUnicode_Check(PyArg);
+#endif	// PY_MAJOR_VERSION < 3
+				if (!bIsStringType && PyUtil::HasLength(PyArg))
+				{
+					const Py_ssize_t SequenceLen = PyObject_Length(PyArg);
+					check(SequenceLen != -1);
+
+					FPyObjectPtr PyObjIter = FPyObjectPtr::StealReference(PyObject_GetIter(PyArg));
+					if (PyObjIter)
+					{
+						// Conversion from a sequence
+						for (Py_ssize_t SequenceIndex = 0; SequenceIndex < SequenceLen; ++SequenceIndex)
+						{
+							FPyObjectPtr ValueItem = FPyObjectPtr::StealReference(PyIter_Next(PyObjIter));
+							if (!ValueItem)
+							{
+								return false;
+							}
+
+							FString& AssetToGenerate = AssetsToGenerate.AddDefaulted_GetRef();
+							if (!PyConversion::Nativize(ValueItem, AssetToGenerate))
+							{
+								PyUtil::SetPythonError(PyExc_TypeError, TEXT("get_blueprint_generated_types"), *FString::Printf(TEXT("Cannot convert argument %d (%s) at index %d to 'string'"), ArgIndex, *PyUtil::GetFriendlyTypename(PyArg), SequenceIndex));
+								return nullptr;
+							}
+						}
+					}
+				}
+				else
+				{
+					FString& AssetToGenerate = AssetsToGenerate.AddDefaulted_GetRef();
+					if (!PyConversion::Nativize(PyArg, AssetToGenerate))
+					{
+						PyUtil::SetPythonError(PyExc_TypeError, TEXT("get_blueprint_generated_types"), *FString::Printf(TEXT("Cannot convert argument %d (%s) to 'string'"), ArgIndex, *PyUtil::GetFriendlyTypename(PyArg)));
+						return nullptr;
+					}
+				}
+			}
+		}
+	}
+
+	if (AssetsToGenerate.Num() == 0)
+	{
+		Py_RETURN_NONE;
+	}
+
+	FPyWrapperTypeRegistry& PyWrapperTypeRegistry = FPyWrapperTypeRegistry::Get();
+	FPyWrapperTypeRegistry::FGeneratedWrappedTypeReferences GeneratedWrappedTypeReferences;
+	TSet<FName> DirtyModules;
+
+	// Process each asset
+	TArray<PyTypeObject*> WrappedAssets;
+	for (const FString& AssetToGenerate : AssetsToGenerate)
+	{
+		const UObject* LoadedAsset = StaticLoadObject(UObject::StaticClass(), nullptr, *AssetToGenerate);
+		if (!LoadedAsset)
+		{
+			PyUtil::SetPythonError(PyExc_RuntimeError, TEXT("get_blueprint_generated_types"), *FString::Printf(TEXT("Cannot find asset for '%s'"), *AssetToGenerate));
+			return nullptr;
+		}
+		LoadedAsset = PyGenUtil::GetTypeRegistryType(LoadedAsset);
+
+		PyTypeObject* WrappedAsset = WrappedAssets.Add_GetRef(PyWrapperTypeRegistry.GenerateWrappedTypeForObject(LoadedAsset, GeneratedWrappedTypeReferences, DirtyModules, EPyTypeGenerationFlags::IncludeBlueprintGeneratedTypes));
+		if (!WrappedAsset)
+		{
+			PyUtil::SetPythonError(PyExc_RuntimeError, TEXT("get_blueprint_generated_types"), *FString::Printf(TEXT("Asset '%s' (%s) cannot be wrapped for Python"), *AssetToGenerate, *LoadedAsset->GetName()));
+			return nullptr;
+		}
+	}
+	check(AssetsToGenerate.Num() == WrappedAssets.Num());
+
+	PyWrapperTypeRegistry.GenerateWrappedTypesForReferences(GeneratedWrappedTypeReferences, DirtyModules);
+	PyWrapperTypeRegistry.NotifyModulesDirtied(DirtyModules);
+
+	check(WrappedAssets.Num() > 0);
+	if (WrappedAssets.Num() == 1)
+	{
+		// Return the single object
+		Py_INCREF(WrappedAssets[0]);
+		return (PyObject*)WrappedAssets[0];
+	}
+	else
+	{
+		// Return the result as a tuple
+		FPyObjectPtr PyWrappedAssets = FPyObjectPtr::StealReference(PyTuple_New(WrappedAssets.Num()));
+		for (int32 WrappedAssetIndex = 0; WrappedAssetIndex < WrappedAssets.Num(); ++WrappedAssetIndex)
+		{
+			Py_INCREF(WrappedAssets[WrappedAssetIndex]);
+			PyTuple_SetItem(PyWrappedAssets, WrappedAssetIndex, (PyObject*)WrappedAssets[WrappedAssetIndex]); // SetItem steals the reference
+		}
+		return PyWrappedAssets.Release();
+	}
+}
+
 PyMethodDef PyEngineMethods[] = {
+	{ "get_blueprint_generated_types", PyCFunctionCast(&GetBlueprintGeneratedTypes), METH_VARARGS, "x.get_blueprint_generated_types(...) -> tuple(type)/type/None -- get the Python types (will return a tuple for multiple types) for the given set of Blueprint asset paths (may be a sequence type or set of arguments)" },
 	{ nullptr, nullptr, 0, nullptr }
 };
 
diff --git a/Engine/Plugins/Experimental/PythonScriptPlugin/Source/PythonScriptPlugin/Private/PyGenUtil.cpp b/Engine/Plugins/Experimental/PythonScriptPlugin/Source/PythonScriptPlugin/Private/PyGenUtil.cpp
--- a/Engine/Plugins/Experimental/PythonScriptPlugin/Source/PythonScriptPlugin/Private/PyGenUtil.cpp
+++ b/Engine/Plugins/Experimental/PythonScriptPlugin/Source/PythonScriptPlugin/Private/PyGenUtil.cpp
@@ -11,6 +11,7 @@
 #include "Misc/Paths.h"
 #include "Misc/ScopeExit.h"
 #include "Misc/FileHelper.h"
+#include "Misc/PackageName.h"
 #include "UObject/Class.h"
 #include "UObject/Stack.h"
 #include "UObject/Package.h"
@@ -19,6 +20,7 @@
 #include "UObject/CoreRedirects.h"
 #include "UObject/PropertyPortFlags.h"
 #include "UObject/UnrealType.h"
+#include "Engine/BlueprintCore.h"
 #include "Engine/BlueprintGeneratedClass.h"
 #include "Engine/UserDefinedStruct.h"
 #include "Engine/UserDefinedEnum.h"
@@ -593,9 +595,20 @@ bool FGeneratedWrappedType::Finalize()
 
 	bool bSuccess = false;
 	// Execute Python code within this block
+	if (FinalizedState != EFinalizedState::Finalized)
 	{
 		FPyScopedGIL GIL;
-		bSuccess = PyType_Ready(&PyType) == 0;
+		if (FinalizedState == EFinalizedState::Initial)
+		{
+			bSuccess = PyType_Ready(&PyType) == 0;
+		}
+		else if (FinalizedState == EFinalizedState::Reset)
+		{
+			PyType_Modified(&PyType);
+			PyDict_SetItemString(PyType.tp_dict, "__doc__", PyUnicode_FromString(PyType.tp_doc ? PyType.tp_doc : "")); // PyType_Modified doesn't update __doc__
+			bSuccess = true;
+		}
+		FinalizedState = EFinalizedState::Finalized;
 	}
 
 	if (bSuccess)
@@ -608,6 +621,12 @@ bool FGeneratedWrappedType::Finalize()
 	return false;
 }
 
+void FGeneratedWrappedType::Reset()
+{
+	Reset_CleansePyType();
+	Reset_CleanseSelf();
+}
+
 void FGeneratedWrappedType::Finalize_PreReady()
 {
 	PyType.tp_name = TypeName.GetData();
@@ -618,6 +637,22 @@ void FGeneratedWrappedType::Finalize_PostReady()
 {
 }
 
+void FGeneratedWrappedType::Reset_CleansePyType()
+{
+	PyType.tp_name = nullptr;
+	PyType.tp_doc = nullptr;
+
+	FPyWrapperBaseMetaData::SetMetaData(&PyType, nullptr);
+}
+
+void FGeneratedWrappedType::Reset_CleanseSelf()
+{
+	TypeName.Reset();
+	TypeDoc.Reset();
+	MetaData.Reset();
+	FinalizedState = EFinalizedState::Reset;
+}
+
 
 void FGeneratedWrappedStructType::Finalize_PreReady()
 {
@@ -657,23 +692,44 @@ void FGeneratedWrappedEnumType::Finalize_PostReady()
 {
 	FGeneratedWrappedType::Finalize_PostReady();
 
-	check(MetaData.IsValid() && MetaData->GetTypeId() == FPyWrapperEnumMetaData::StaticTypeId());
-	TSharedRef<FPyWrapperEnumMetaData> EnumMetaData = StaticCastSharedRef<FPyWrapperEnumMetaData>(MetaData.ToSharedRef());
+	check(!MetaData || MetaData->GetTypeId() == FPyWrapperEnumMetaData::StaticTypeId());
+	TSharedPtr<FPyWrapperEnumMetaData> EnumMetaData = StaticCastSharedPtr<FPyWrapperEnumMetaData>(MetaData);
 
-	// Execute Python code within this block
+	if (EnumMetaData)
 	{
-		FPyScopedGIL GIL;
-		for (const FGeneratedWrappedEnumEntry& EnumEntry : EnumEntries)
+		// Execute Python code within this block
 		{
-			FPyWrapperEnum* PyEnumEntry = FPyWrapperEnum::AddEnumEntry(&PyType, EnumEntry.EntryValue, EnumEntry.EntryName.GetData(), EnumEntry.EntryDoc.GetData());
-			if (PyEnumEntry)
+			FPyScopedGIL GIL;
+			for (const FGeneratedWrappedEnumEntry& EnumEntry : EnumEntries)
 			{
-				EnumMetaData->EnumEntries.Add(PyEnumEntry);
+				FPyWrapperEnum* PyEnumEntry = FPyWrapperEnum::AddEnumEntry(&PyType, EnumEntry.EntryValue, EnumEntry.EntryName.GetData(), EnumEntry.EntryDoc.GetData());
+				if (PyEnumEntry)
+				{
+					EnumMetaData->EnumEntries.Add(PyEnumEntry);
+				}
 			}
 		}
+
+		EnumMetaData->bFinalized = true;
+	}
+}
+
+void FGeneratedWrappedEnumType::Reset_CleansePyType()
+{
+	// Unregister the existing enum entries
+	for (const FGeneratedWrappedEnumEntry& EnumEntry : EnumEntries)
+	{
+		PyDict_DelItemString(PyType.tp_dict, EnumEntry.EntryName.GetData());
 	}
 
-	EnumMetaData->bFinalized = true;
+	FGeneratedWrappedType::Reset_CleansePyType();
+}
+
+void FGeneratedWrappedEnumType::Reset_CleanseSelf()
+{
+	EnumEntries.Reset();
+
+	FGeneratedWrappedType::Reset_CleanseSelf();
 }
 
 void FGeneratedWrappedEnumType::ExtractEnumEntries(const UEnum* InEnum)
@@ -1281,6 +1337,11 @@ bool IsScriptExposedClass(const UClass* InClass)
 {
 	for (const UClass* ParentClass = InClass; ParentClass; ParentClass = ParentClass->GetSuperClass())
 	{
+		if (IsBlueprintGeneratedClass(ParentClass))
+		{
+			return true;
+		}
+
 		if (ParentClass->GetBoolMetaData(BlueprintTypeMetaDataKey) || ParentClass->HasMetaData(BlueprintSpawnableComponentMetaDataKey))
 		{
 			return true;
@@ -1299,6 +1360,11 @@ bool IsScriptExposedStruct(const UScriptStruct* InStruct)
 {
 	for (const UScriptStruct* ParentStruct = InStruct; ParentStruct; ParentStruct = Cast<UScriptStruct>(ParentStruct->GetSuperStruct()))
 	{
+		if (IsBlueprintGeneratedStruct(ParentStruct))
+		{
+			return true;
+		}
+
 		if (ParentStruct->GetBoolMetaData(BlueprintTypeMetaDataKey))
 		{
 			return true;
@@ -1315,6 +1381,11 @@ bool IsScriptExposedStruct(const UScriptStruct* InStruct)
 
 bool IsScriptExposedEnum(const UEnum* InEnum)
 {
+	if (IsBlueprintGeneratedEnum(InEnum))
+	{
+		return true;
+	}
+
 	if (InEnum->GetBoolMetaData(BlueprintTypeMetaDataKey))
 	{
 		return true;
@@ -2243,6 +2314,42 @@ FString PythonizeDefaultValue(const UProperty* InProp, const FString& InDefaultV
 	return PythonizeValue(InProp, PropValue.GetValue(), InFlags);
 }
 
+const UObject* GetTypeRegistryType(const UObject* InObj)
+{
+	if (const UBlueprintCore* BlueprintAsset = Cast<const UBlueprintCore>(InObj))
+	{
+		return BlueprintAsset->GeneratedClass;
+	}
+
+	return InObj;
+}
+
+FName GetTypeRegistryName(const UClass* InClass)
+{
+	return IsBlueprintGeneratedClass(InClass)
+		? InClass->GetOutermost()->GetFName()
+		: InClass->GetFName();
+}
+
+FName GetTypeRegistryName(const UScriptStruct* InStruct)
+{
+	return IsBlueprintGeneratedStruct(InStruct)
+		? InStruct->GetOutermost()->GetFName()
+		: InStruct->GetFName();
+}
+
+FName GetTypeRegistryName(const UEnum* InEnum)
+{
+	return IsBlueprintGeneratedEnum(InEnum)
+		? InEnum->GetOutermost()->GetFName()
+		: InEnum->GetFName();
+}
+
+FName GetTypeRegistryName(const UFunction* InDelegateSignature)
+{
+	return InDelegateSignature->GetFName();
+}
+
 FString GetFieldModule(const UField* InField)
 {
 	UPackage* ScriptPackage = InField->GetOutermost();
@@ -2252,6 +2359,11 @@ FString GetFieldModule(const UField* InField)
 	{
 		return PackageName.RightChop(8); // Chop "/Script/" from the name
 	}
+	else
+	{
+		// Not a native module!
+		return FString();
+	}
 
 	check(PackageName[0] == TEXT('/'));
 	int32 RootNameEnd = 1;
@@ -2505,6 +2617,28 @@ FString GetEnumEntryPythonName(const UEnum* InEnum, const int32 InEntryIndex)
 		}
 	}
 	
+	// Blueprint enums have horrible internal names, so try and create a valid identifier from the display name meta-data
+	if (IsBlueprintGeneratedEnum(InEnum))
+	{
+		EnumEntryName = InEnum->GetAuthoredNameStringByIndex(InEntryIndex);
+		if (FCString::IsPureAnsi(*EnumEntryName))
+		{
+			// Sanitize out any invalid characters in the name
+			for (TCHAR& Char : EnumEntryName)
+			{
+				if (!FChar::IsAlnum(Char))
+				{
+					Char = TEXT('_');
+				}
+			}
+		}
+		else
+		{
+			// If the name isn't pure-ANSI then don't attempt to sanitize it as the result will be very mangled
+			EnumEntryName.Reset();
+		}
+	}
+
 	// Just use the entry name if we have no meta-data
 	if (EnumEntryName.IsEmpty())
 	{
@@ -2741,14 +2875,14 @@ FString GetEnumEntryTooltip(const UEnum* InEnum, const int64 InEntryIndex)
 	return *FTextInspector::GetSourceString(InEnum->GetToolTipTextByIndex(InEntryIndex));
 }
 
-FString BuildCppSourceInformationDocString(const UField* InOwnerType)
+FString BuildSourceInformationDocString(const UField* InOwnerType)
 {
 	FString Str;
-	AppendCppSourceInformationDocString(InOwnerType, Str);
+	AppendSourceInformationDocString(InOwnerType, Str);
 	return Str;
 }
 
-void AppendCppSourceInformationDocString(const UField* InOwnerType, FString& OutStr)
+void AppendSourceInformationDocString(const UField* InOwnerType, FString& OutStr)
 {
 	if (!InOwnerType)
 	{
@@ -2763,23 +2897,35 @@ void AppendCppSourceInformationDocString(const UField* InOwnerType, FString& Out
 		}
 	}
 
-	static const FName ModuleRelativePathMetaDataKey = "ModuleRelativePath";
+	const UPackage* OwnerPackage = InOwnerType->GetOutermost();
+	const FString OwnerPackageName = OwnerPackage->GetName();
+
+	if (FPackageName::IsScriptPackage(OwnerPackageName))
+	{
+		static const FName ModuleRelativePathMetaDataKey = "ModuleRelativePath";
 
-	const FString TypePlugin = GetFieldPlugin(InOwnerType);
-	const FString TypeModule = GetFieldModule(InOwnerType);
-	const FString TypeFile = FPaths::GetCleanFilename(InOwnerType->GetMetaData(ModuleRelativePathMetaDataKey));
+		const FString TypePlugin = GetFieldPlugin(InOwnerType);
+		const FString TypeModule = GetFieldModule(InOwnerType);
+		const FString TypeFile = FPaths::GetCleanFilename(InOwnerType->GetMetaData(ModuleRelativePathMetaDataKey));
 
-	OutStr += LINE_TERMINATOR TEXT("**C++ Source:**") LINE_TERMINATOR;
-	if (!TypePlugin.IsEmpty())
+		OutStr += LINE_TERMINATOR TEXT("**C++ Source:**") LINE_TERMINATOR;
+		if (!TypePlugin.IsEmpty())
+		{
+			OutStr += LINE_TERMINATOR TEXT("- **Plugin**: ");
+			OutStr += TypePlugin;
+		}
+		OutStr += LINE_TERMINATOR TEXT("- **Module**: ");
+		OutStr += TypeModule;
+		OutStr += LINE_TERMINATOR TEXT("- **File**: ");
+		OutStr += TypeFile;
+		OutStr += LINE_TERMINATOR;
+	}
+	else
 	{
-		OutStr += LINE_TERMINATOR TEXT("- **Plugin**: ");
-		OutStr += TypePlugin;
+		OutStr += LINE_TERMINATOR TEXT("**Asset Source:**") LINE_TERMINATOR;
+		OutStr += LINE_TERMINATOR TEXT("- **Package**: ");
+		OutStr += OwnerPackageName;
 	}
-	OutStr += LINE_TERMINATOR TEXT("- **Module**: ");
-	OutStr += TypeModule;
-	OutStr += LINE_TERMINATOR TEXT("- **File**: ");
-	OutStr += TypeFile;
-	OutStr += LINE_TERMINATOR;
 }
 
 bool SaveGeneratedTextFile(const TCHAR* InFilename, const FString& InFileContents, const bool InForceWrite)
diff --git a/Engine/Plugins/Experimental/PythonScriptPlugin/Source/PythonScriptPlugin/Private/PyGenUtil.h b/Engine/Plugins/Experimental/PythonScriptPlugin/Source/PythonScriptPlugin/Private/PyGenUtil.h
--- a/Engine/Plugins/Experimental/PythonScriptPlugin/Source/PythonScriptPlugin/Private/PyGenUtil.h
+++ b/Engine/Plugins/Experimental/PythonScriptPlugin/Source/PythonScriptPlugin/Private/PyGenUtil.h
@@ -666,6 +666,13 @@ namespace PyGenUtil
 	/** Stores the minimal data needed by a runtime generated Python type */
 	struct FGeneratedWrappedType
 	{
+		enum class EFinalizedState : uint8
+		{
+			Initial,
+			Reset,
+			Finalized,
+		};
+
 		FGeneratedWrappedType()
 		{
 			PyType = { PyVarObject_HEAD_INIT(nullptr, 0) };
@@ -681,12 +688,21 @@ namespace PyGenUtil
 		/** Called to ready the generated type with Python */
 		bool Finalize();
 
+		/** Called to reset this type back to its clean state */
+		void Reset();
+
 		/** Internal version of Finalize, called before readying the type with Python */
 		virtual void Finalize_PreReady();
 
 		/** Internal version of Finalize, called after readying the type with Python */
 		virtual void Finalize_PostReady();
 
+		/** Internal version of Reset, called to remove data added to the Python type */
+		virtual void Reset_CleansePyType();
+
+		/** Internal version of Reset, called to reset the data on this type */
+		virtual void Reset_CleanseSelf();
+
 		/** The name of the type */
 		FUTF8Buffer TypeName;
 
@@ -698,6 +714,9 @@ namespace PyGenUtil
 
 		/* The Python type that was generated */
 		PyTypeObject PyType;
+
+		/** What is the current finalization state of this generated type? */
+		EFinalizedState FinalizedState = EFinalizedState::Initial;
 	};
 
 	/** Stores the data needed by a runtime generated Python struct type */
@@ -784,6 +803,12 @@ namespace PyGenUtil
 		/** Internal version of Finalize, called after readying the type with Python */
 		virtual void Finalize_PostReady() override;
 
+		/** Internal version of Reset, called to remove data added to the Python type */
+		virtual void Reset_CleansePyType() override;
+
+		/** Internal version of Reset, called to reset the data on this type */
+		virtual void Reset_CleanseSelf() override;
+
 		/** Called to extract the enum entries array from the given enum */
 		void ExtractEnumEntries(const UEnum* InEnum);
 
@@ -977,6 +1002,21 @@ namespace PyGenUtil
 	/** Given a property and its default value, convert it into something that could be used in a Python script */
 	FString PythonizeDefaultValue(const UProperty* InProp, const FString& InDefaultValue, const uint32 InFlags = EPythonizeValueFlags::None);
 
+	/** Get the type that should be used with the Python type registry (eg, a Blueprint asset should use its generated class) */
+	const UObject* GetTypeRegistryType(const UObject* InObj);
+
+	/** Get the name that should be used by the given class when registered with the Python type registry */
+	FName GetTypeRegistryName(const UClass* InClass);
+
+	/** Get the name that should be used by the given struct when registered with the Python type registry */
+	FName GetTypeRegistryName(const UScriptStruct* InStruct);
+
+	/** Get the name that should be used by the given enum when registered with the Python type registry */
+	FName GetTypeRegistryName(const UEnum* InEnum);
+
+	/** Get the name that should be used by the given delegate when registered with the Python type registry */
+	FName GetTypeRegistryName(const UFunction* InDelegateSignature);
+
 	/** Get the native module the given field belongs to */
 	FString GetFieldModule(const UField* InField);
 
@@ -1052,11 +1092,11 @@ namespace PyGenUtil
 	/** Get the tooltip for the given enum entry */
 	FString GetEnumEntryTooltip(const UEnum* InEnum, const int64 InEntryIndex);
 
-	/** Get the doc string for the C++ source information of the given type */
-	FString BuildCppSourceInformationDocString(const UField* InOwnerType);
+	/** Get the doc string for the C++ or asset source information of the given type */
+	FString BuildSourceInformationDocString(const UField* InOwnerType);
 
-	/** Append the doc string for the C++ source information of the given type to the given string */
-	void AppendCppSourceInformationDocString(const UField* InOwnerType, FString& OutStr);
+	/** Append the doc string for the C++ or asset source information of the given type to the given string */
+	void AppendSourceInformationDocString(const UField* InOwnerType, FString& OutStr);
 
 	/** Save a generated text file to disk as UTF-8 (only writes the file if the contents differs, unless forced) */
 	bool SaveGeneratedTextFile(const TCHAR* InFilename, const FString& InFileContents, const bool InForceWrite = false);
diff --git a/Engine/Plugins/Experimental/PythonScriptPlugin/Source/PythonScriptPlugin/Private/PyWrapperDelegate.h b/Engine/Plugins/Experimental/PythonScriptPlugin/Source/PythonScriptPlugin/Private/PyWrapperDelegate.h
--- a/Engine/Plugins/Experimental/PythonScriptPlugin/Source/PythonScriptPlugin/Private/PyWrapperDelegate.h
+++ b/Engine/Plugins/Experimental/PythonScriptPlugin/Source/PythonScriptPlugin/Private/PyWrapperDelegate.h
@@ -115,6 +115,13 @@ struct TPyWrapperDelegateMetaData : public FPyWrapperBaseMetaData
 		return GetPythonCallableForDelegateClass(Py_TYPE(Instance));
 	}
 
+	/** Add object references from the given Python object to the given collector */
+	virtual void AddReferencedObjects(FPyWrapperBase* Instance, FReferenceCollector& Collector) override
+	{
+		Collector.AddReferencedObject(DelegateSignature.Func);
+		Collector.AddReferencedObject(PythonCallableForDelegateClass);
+	}
+
 	/** Get the reflection meta data type object associated with this wrapper type if there is one or nullptr if not. */
 	virtual const UField* GetMetaType() const override
 	{
diff --git a/Engine/Plugins/Experimental/PythonScriptPlugin/Source/PythonScriptPlugin/Private/PyWrapperEnum.cpp b/Engine/Plugins/Experimental/PythonScriptPlugin/Source/PythonScriptPlugin/Private/PyWrapperEnum.cpp
--- a/Engine/Plugins/Experimental/PythonScriptPlugin/Source/PythonScriptPlugin/Private/PyWrapperEnum.cpp
+++ b/Engine/Plugins/Experimental/PythonScriptPlugin/Source/PythonScriptPlugin/Private/PyWrapperEnum.cpp
@@ -714,6 +714,11 @@ bool FPyWrapperEnumMetaData::IsEnumFinalized(FPyWrapperEnum* Instance)
 	return IsEnumFinalized(Py_TYPE(Instance));
 }
 
+void FPyWrapperEnumMetaData::AddReferencedObjects(FPyWrapperBase* Instance, FReferenceCollector& Collector)
+{
+	Collector.AddReferencedObject(Enum);
+}
+
 class FPythonGeneratedEnumBuilder
 {
 public:
diff --git a/Engine/Plugins/Experimental/PythonScriptPlugin/Source/PythonScriptPlugin/Private/PyWrapperEnum.h b/Engine/Plugins/Experimental/PythonScriptPlugin/Source/PythonScriptPlugin/Private/PyWrapperEnum.h
--- a/Engine/Plugins/Experimental/PythonScriptPlugin/Source/PythonScriptPlugin/Private/PyWrapperEnum.h
+++ b/Engine/Plugins/Experimental/PythonScriptPlugin/Source/PythonScriptPlugin/Private/PyWrapperEnum.h
@@ -85,6 +85,9 @@ struct FPyWrapperEnumMetaData : public FPyWrapperBaseMetaData
 	/** Check to see if the enum is finalized */
 	static bool IsEnumFinalized(FPyWrapperEnum* Instance);
 
+	/** Add object references from the given Python object to the given collector */
+	virtual void AddReferencedObjects(FPyWrapperBase* Instance, FReferenceCollector& Collector) override;
+
 	/** Get the reflection meta data type object associated with this wrapper type if there is one or nullptr if not. */
 	virtual const UField* GetMetaType() const override
 	{
@@ -97,7 +100,7 @@ struct FPyWrapperEnumMetaData : public FPyWrapperBaseMetaData
 	/** True if this enum type has been finalized after having all of its entries added to it */
 	bool bFinalized;
 
-	/** Set if this struct is deprecated and using it should emit a deprecation warning */
+	/** Set if this enum is deprecated and using it should emit a deprecation warning */
 	TOptional<FString> DeprecationMessage;
 
 	/** Array of enum entries in the order they were added (enum entries are stored as borrowed references) */
diff --git a/Engine/Plugins/Experimental/PythonScriptPlugin/Source/PythonScriptPlugin/Private/PyWrapperObject.cpp b/Engine/Plugins/Experimental/PythonScriptPlugin/Source/PythonScriptPlugin/Private/PyWrapperObject.cpp
--- a/Engine/Plugins/Experimental/PythonScriptPlugin/Source/PythonScriptPlugin/Private/PyWrapperObject.cpp
+++ b/Engine/Plugins/Experimental/PythonScriptPlugin/Source/PythonScriptPlugin/Private/PyWrapperObject.cpp
@@ -912,6 +912,8 @@ void FPyWrapperObjectMetaData::AddReferencedObjects(FPyWrapperBase* Instance, FR
 {
 	FPyWrapperObject* Self = static_cast<FPyWrapperObject*>(Instance);
 
+	Collector.AddReferencedObject(Class);
+
 	UObject* OldInstance = Self->ObjectInstance;
 	Collector.AddReferencedObject(Self->ObjectInstance);
 
diff --git a/Engine/Plugins/Experimental/PythonScriptPlugin/Source/PythonScriptPlugin/Private/PyWrapperStruct.cpp b/Engine/Plugins/Experimental/PythonScriptPlugin/Source/PythonScriptPlugin/Private/PyWrapperStruct.cpp
--- a/Engine/Plugins/Experimental/PythonScriptPlugin/Source/PythonScriptPlugin/Private/PyWrapperStruct.cpp
+++ b/Engine/Plugins/Experimental/PythonScriptPlugin/Source/PythonScriptPlugin/Private/PyWrapperStruct.cpp
@@ -1183,10 +1183,12 @@ FPyWrapperStructMetaData::FPyWrapperStructMetaData()
 {
 }
 
-/** Add object references from the given Python object to the given collector */
 void FPyWrapperStructMetaData::AddReferencedObjects(FPyWrapperBase* Instance, FReferenceCollector& Collector)
 {
 	FPyWrapperStruct* Self = static_cast<FPyWrapperStruct*>(Instance);
+
+	Collector.AddReferencedObject(Struct);
+	
 	Collector.AddReferencedObject(Self->ScriptStruct);
 	if (Self->ScriptStruct && Self->StructInstance && !Self->OwnerContext.HasOwner())
 	{
diff --git a/Engine/Plugins/Experimental/PythonScriptPlugin/Source/PythonScriptPlugin/Private/PyWrapperTypeRegistry.cpp b/Engine/Plugins/Experimental/PythonScriptPlugin/Source/PythonScriptPlugin/Private/PyWrapperTypeRegistry.cpp
--- a/Engine/Plugins/Experimental/PythonScriptPlugin/Source/PythonScriptPlugin/Private/PyWrapperTypeRegistry.cpp
+++ b/Engine/Plugins/Experimental/PythonScriptPlugin/Source/PythonScriptPlugin/Private/PyWrapperTypeRegistry.cpp
@@ -522,6 +522,8 @@ void FPyWrapperTypeRegistry::OrphanWrappedTypesForModule(const FName ModuleName)
 
 void FPyWrapperTypeRegistry::GenerateWrappedTypesForReferences(const FGeneratedWrappedTypeReferences& InGeneratedWrappedTypeReferences, TSet<FName>& OutDirtyModules)
 {
+	static const EPyTypeGenerationFlags ReferenceGenerationFlags = EPyTypeGenerationFlags::ForceShouldExport | EPyTypeGenerationFlags::IncludeBlueprintGeneratedTypes;
+
 	if (!InGeneratedWrappedTypeReferences.HasReferences())
 	{
 		return;
@@ -531,23 +533,23 @@ void FPyWrapperTypeRegistry::GenerateWrappedTypesForReferences(const FGeneratedW
 
 	for (const UClass* Class : InGeneratedWrappedTypeReferences.ClassReferences)
 	{
-		GenerateWrappedClassType(Class, GeneratedWrappedTypeReferences, OutDirtyModules, true);
+		GenerateWrappedClassType(Class, GeneratedWrappedTypeReferences, OutDirtyModules, ReferenceGenerationFlags);
 	}
 
 	for (const UScriptStruct* Struct : InGeneratedWrappedTypeReferences.StructReferences)
 	{
-		GenerateWrappedStructType(Struct, GeneratedWrappedTypeReferences, OutDirtyModules, true);
+		GenerateWrappedStructType(Struct, GeneratedWrappedTypeReferences, OutDirtyModules, ReferenceGenerationFlags);
 	}
 
 	for (const UEnum* Enum : InGeneratedWrappedTypeReferences.EnumReferences)
 	{
-		GenerateWrappedEnumType(Enum, GeneratedWrappedTypeReferences, OutDirtyModules, true);
+		GenerateWrappedEnumType(Enum, GeneratedWrappedTypeReferences, OutDirtyModules, ReferenceGenerationFlags);
 	}
 
 	for (const UFunction* DelegateSignature : InGeneratedWrappedTypeReferences.DelegateReferences)
 	{
 		check(DelegateSignature->HasAnyFunctionFlags(FUNC_Delegate));
-		GenerateWrappedDelegateType(DelegateSignature, GeneratedWrappedTypeReferences, OutDirtyModules, true);
+		GenerateWrappedDelegateType(DelegateSignature, GeneratedWrappedTypeReferences, OutDirtyModules, ReferenceGenerationFlags);
 	}
 
 	GenerateWrappedTypesForReferences(GeneratedWrappedTypeReferences, OutDirtyModules);
@@ -562,53 +564,191 @@ void FPyWrapperTypeRegistry::NotifyModulesDirtied(const TSet<FName>& InDirtyModu
 	}
 }
 
-PyTypeObject* FPyWrapperTypeRegistry::GenerateWrappedTypeForObject(const UObject* InObj, FGeneratedWrappedTypeReferences& OutGeneratedWrappedTypeReferences, TSet<FName>& OutDirtyModules, const bool bForceGenerate)
+void FPyWrapperTypeRegistry::UpdateGenerateWrappedTypeForRename(const FName InOldName, const UObject* InObj)
+{
+	FName TypeRegistryName;
+	{
+		if (const UClass* Class = Cast<const UClass>(InObj))
+		{
+			TypeRegistryName = PyGenUtil::GetTypeRegistryName(Class);
+		}
+
+		if (const UScriptStruct* Struct = Cast<const UScriptStruct>(InObj))
+		{
+			TypeRegistryName = PyGenUtil::GetTypeRegistryName(Struct);
+		}
+
+		if (const UEnum* Enum = Cast<const UEnum>(InObj))
+		{
+			TypeRegistryName = PyGenUtil::GetTypeRegistryName(Enum);
+		}
+
+		if (const UFunction* Func = Cast<const UFunction>(InObj))
+		{
+			if (Func->HasAnyFunctionFlags(FUNC_Delegate))
+			{
+				TypeRegistryName = PyGenUtil::GetTypeRegistryName(Func);
+			}
+		}
+	}
+
+	TSharedPtr<PyGenUtil::FGeneratedWrappedType> GeneratedWrappedType = GeneratedWrappedTypes.FindAndRemoveChecked(InOldName);
+	check(!TypeRegistryName.IsNone() && !GeneratedWrappedTypes.Contains(TypeRegistryName));
+	GeneratedWrappedTypes.Add(TypeRegistryName, GeneratedWrappedType);
+}
+
+void FPyWrapperTypeRegistry::RemoveGenerateWrappedTypeForDelete(const UObject* InObj)
+{
+	FName TypeRegistryName;
+	FString PythonTypeName;
+	{
+		if (const UClass* Class = Cast<const UClass>(InObj))
+		{
+			TypeRegistryName = PyGenUtil::GetTypeRegistryName(Class);
+			PythonTypeName = PyGenUtil::GetClassPythonName(Class);
+		}
+
+		if (const UScriptStruct* Struct = Cast<const UScriptStruct>(InObj))
+		{
+			TypeRegistryName = PyGenUtil::GetTypeRegistryName(Struct);
+			PythonTypeName = PyGenUtil::GetStructPythonName(Struct);
+		}
+
+		if (const UEnum* Enum = Cast<const UEnum>(InObj))
+		{
+			TypeRegistryName = PyGenUtil::GetTypeRegistryName(Enum);
+			PythonTypeName = PyGenUtil::GetEnumPythonName(Enum);
+		}
+
+		if (const UFunction* Func = Cast<const UFunction>(InObj))
+		{
+			if (Func->HasAnyFunctionFlags(FUNC_Delegate))
+			{
+				TypeRegistryName = PyGenUtil::GetTypeRegistryName(Func);
+				PythonTypeName = PyGenUtil::GetDelegatePythonName(Func);
+			}
+		}
+	}
+
+	TSharedPtr<PyGenUtil::FGeneratedWrappedType> GeneratedWrappedType = GeneratedWrappedTypes.FindRef(TypeRegistryName);
+	GeneratedWrappedType->Reset();
+
+	// Fill the type with dummy information and re-finalize it as an empty stub
+	GeneratedWrappedType->TypeName = PyGenUtil::TCHARToUTF8Buffer(*FString::Printf(TEXT("%s_DELETED"), *PythonTypeName));
+	GeneratedWrappedType->TypeDoc = PyGenUtil::TCHARToUTF8Buffer(*FString::Printf(TEXT("%s was deleted!"), *InObj->GetPathName()));
+	GeneratedWrappedType->Finalize();
+}
+
+PyTypeObject* FPyWrapperTypeRegistry::GenerateWrappedTypeForObject(const UObject* InObj, FGeneratedWrappedTypeReferences& OutGeneratedWrappedTypeReferences, TSet<FName>& OutDirtyModules, const EPyTypeGenerationFlags InGenerationFlags)
+{
+	if (const UClass* Class = Cast<const UClass>(InObj))
+	{
+		return GenerateWrappedClassType(Class, OutGeneratedWrappedTypeReferences, OutDirtyModules, InGenerationFlags);
+	}
+
+	if (const UScriptStruct* Struct = Cast<const UScriptStruct>(InObj))
+	{
+		return GenerateWrappedStructType(Struct, OutGeneratedWrappedTypeReferences, OutDirtyModules, InGenerationFlags);
+	}
+
+	if (const UEnum* Enum = Cast<const UEnum>(InObj))
+	{
+		return GenerateWrappedEnumType(Enum, OutGeneratedWrappedTypeReferences, OutDirtyModules, InGenerationFlags);
+	}
+
+	if (const UFunction* Func = Cast<const UFunction>(InObj))
+	{
+		if (Func->HasAnyFunctionFlags(FUNC_Delegate))
+		{
+			return GenerateWrappedDelegateType(Func, OutGeneratedWrappedTypeReferences, OutDirtyModules, InGenerationFlags);
+		}
+	}
+
+	return nullptr;
+}
+
+bool FPyWrapperTypeRegistry::HasWrappedTypeForObject(const UObject* InObj) const
+{
+	if (const UClass* Class = Cast<const UClass>(InObj))
+	{
+		return HasWrappedClassType(Class);
+	}
+
+	if (const UScriptStruct* Struct = Cast<const UScriptStruct>(InObj))
+	{
+		return HasWrappedStructType(Struct);
+	}
+
+	if (const UEnum* Enum = Cast<const UEnum>(InObj))
+	{
+		return HasWrappedEnumType(Enum);
+	}
+
+	if (const UFunction* Func = Cast<const UFunction>(InObj))
+	{
+		if (Func->HasAnyFunctionFlags(FUNC_Delegate))
+		{
+			return HasWrappedDelegateType(Func);
+		}
+	}
+
+	return false;
+}
+
+bool FPyWrapperTypeRegistry::HasWrappedTypeForObjectName(const FName InName) const
+{
+	return GeneratedWrappedTypes.Contains(InName);
+}
+
+PyTypeObject* FPyWrapperTypeRegistry::GetWrappedTypeForObject(const UObject* InObj) const
 {
 	if (const UClass* Class = Cast<const UClass>(InObj))
 	{
-		return GenerateWrappedClassType(Class, OutGeneratedWrappedTypeReferences, OutDirtyModules, bForceGenerate);
+		return GetWrappedClassType(Class);
 	}
 
 	if (const UScriptStruct* Struct = Cast<const UScriptStruct>(InObj))
 	{
-		return GenerateWrappedStructType(Struct, OutGeneratedWrappedTypeReferences, OutDirtyModules, bForceGenerate);
+		return GetWrappedStructType(Struct);
 	}
 
 	if (const UEnum* Enum = Cast<const UEnum>(InObj))
 	{
-		return GenerateWrappedEnumType(Enum, OutGeneratedWrappedTypeReferences, OutDirtyModules, bForceGenerate);
+		return GetWrappedEnumType(Enum);
 	}
 
 	if (const UFunction* Func = Cast<const UFunction>(InObj))
 	{
 		if (Func->HasAnyFunctionFlags(FUNC_Delegate))
 		{
-			return GenerateWrappedDelegateType(Func, OutGeneratedWrappedTypeReferences, OutDirtyModules, bForceGenerate);
+			return GetWrappedDelegateType(Func);
 		}
 	}
 
 	return nullptr;
 }
 
-PyTypeObject* FPyWrapperTypeRegistry::GenerateWrappedClassType(const UClass* InClass, FGeneratedWrappedTypeReferences& OutGeneratedWrappedTypeReferences, TSet<FName>& OutDirtyModules, const bool bForceGenerate)
+PyTypeObject* FPyWrapperTypeRegistry::GenerateWrappedClassType(const UClass* InClass, FGeneratedWrappedTypeReferences& OutGeneratedWrappedTypeReferences, TSet<FName>& OutDirtyModules, const EPyTypeGenerationFlags InGenerationFlags)
 {
 	SCOPE_SECONDS_ACCUMULATOR(STAT_GenerateWrappedClassTotalTime);
 	INC_DWORD_STAT(STAT_GenerateWrappedClassCallCount);
 
 	// Already processed? Nothing more to do
-	if (PyTypeObject* ExistingPyType = PythonWrappedClasses.FindRef(InClass->GetFName()))
+	const FName TypeRegistryName = PyGenUtil::GetTypeRegistryName(InClass);
+	if (PyTypeObject* ExistingPyType = PythonWrappedClasses.FindRef(TypeRegistryName))
 	{
 		return ExistingPyType;
 	}
 
 	// todo: allow generation of Blueprint generated classes
-	if (PyGenUtil::IsBlueprintGeneratedClass(InClass))
+	const bool bIsBlueprintGeneratedType = PyGenUtil::IsBlueprintGeneratedClass(InClass);
+	if (bIsBlueprintGeneratedType)
 	{
 		return nullptr;
 	}
 
 	// Should this type be exported?
-	if (!bForceGenerate && !PyGenUtil::ShouldExportClass(InClass))
+	if (!EnumHasAnyFlags(InGenerationFlags, EPyTypeGenerationFlags::ForceShouldExport) && !PyGenUtil::ShouldExportClass(InClass))
 	{
 		return nullptr;
 	}
@@ -617,14 +757,14 @@ PyTypeObject* FPyWrapperTypeRegistry::GenerateWrappedClassType(const UClass* InC
 	PyTypeObject* SuperPyType = nullptr;
 	if (const UClass* SuperClass = InClass->GetSuperClass())
 	{
-		SuperPyType = GenerateWrappedClassType(SuperClass, OutGeneratedWrappedTypeReferences, OutDirtyModules, true);
+		SuperPyType = GenerateWrappedClassType(SuperClass, OutGeneratedWrappedTypeReferences, OutDirtyModules, InGenerationFlags | EPyTypeGenerationFlags::ForceShouldExport);
 	}
 
 	INC_DWORD_STAT(STAT_GenerateWrappedClassObjCount);
 
-	check(!GeneratedWrappedTypes.Contains(InClass->GetFName()));
+	check(!GeneratedWrappedTypes.Contains(TypeRegistryName));
 	TSharedRef<PyGenUtil::FGeneratedWrappedClassType> GeneratedWrappedType = MakeShared<PyGenUtil::FGeneratedWrappedClassType>();
-	GeneratedWrappedTypes.Add(InClass->GetFName(), GeneratedWrappedType);
+	GeneratedWrappedTypes.Add(TypeRegistryName, GeneratedWrappedType);
 
 	TMap<FName, FName> PythonProperties;
 	TMap<FName, FString> PythonDeprecatedProperties;
@@ -850,13 +990,13 @@ PyTypeObject* FPyWrapperTypeRegistry::GenerateWrappedClassType(const UClass* InC
 		{
 			// Ensure that we've generated a finalized Python type for this class since we'll be adding this function as a dynamic method on that type
 			const UClass* HostedClass = CastChecked<UObjectPropertyBase>(SelfParam.ParamProp)->PropertyClass;
-			if (!GenerateWrappedClassType(HostedClass, OutGeneratedWrappedTypeReferences, OutDirtyModules, true))
+			if (!GenerateWrappedClassType(HostedClass, OutGeneratedWrappedTypeReferences, OutDirtyModules, EPyTypeGenerationFlags::ForceShouldExport))
 			{
 				return;
 			}
 
 			// Find the wrapped type for the class as that's what we'll actually add the dynamic method to
-			TSharedPtr<PyGenUtil::FGeneratedWrappedClassType> HostedClassGeneratedWrappedType = StaticCastSharedPtr<PyGenUtil::FGeneratedWrappedClassType>(GeneratedWrappedTypes.FindRef(HostedClass->GetFName()));
+			TSharedPtr<PyGenUtil::FGeneratedWrappedClassType> HostedClassGeneratedWrappedType = StaticCastSharedPtr<PyGenUtil::FGeneratedWrappedClassType>(GeneratedWrappedTypes.FindRef(PyGenUtil::GetTypeRegistryName(HostedClass)));
 			check(HostedClassGeneratedWrappedType.IsValid());
 
 			// Add the dynamic methods to the class
@@ -870,13 +1010,13 @@ PyTypeObject* FPyWrapperTypeRegistry::GenerateWrappedClassType(const UClass* InC
 		{
 			// Ensure that we've generated a finalized Python type for this struct since we'll be adding this function as a dynamic method on that type
 			const UScriptStruct* HostedStruct = CastChecked<UStructProperty>(SelfParam.ParamProp)->Struct;
-			if (!GenerateWrappedStructType(HostedStruct, OutGeneratedWrappedTypeReferences, OutDirtyModules, true))
+			if (!GenerateWrappedStructType(HostedStruct, OutGeneratedWrappedTypeReferences, OutDirtyModules, EPyTypeGenerationFlags::ForceShouldExport))
 			{
 				return;
 			}
 
 			// Find the wrapped type for the struct as that's what we'll actually add the dynamic method to
-			TSharedPtr<PyGenUtil::FGeneratedWrappedStructType> HostedStructGeneratedWrappedType = StaticCastSharedPtr<PyGenUtil::FGeneratedWrappedStructType>(GeneratedWrappedTypes.FindRef(HostedStruct->GetFName()));
+			TSharedPtr<PyGenUtil::FGeneratedWrappedStructType> HostedStructGeneratedWrappedType = StaticCastSharedPtr<PyGenUtil::FGeneratedWrappedStructType>(GeneratedWrappedTypes.FindRef(PyGenUtil::GetTypeRegistryName(HostedStruct)));
 			check(HostedStructGeneratedWrappedType.IsValid());
 
 			// Add the dynamic methods to the struct
@@ -930,10 +1070,10 @@ PyTypeObject* FPyWrapperTypeRegistry::GenerateWrappedClassType(const UClass* InC
 
 			// Ensure that we've generated a finalized Python type for this struct since we'll be adding this function as a operator on that type
 			const UScriptStruct* HostedStruct = CastChecked<UStructProperty>(OpFunc.SelfParam.ParamProp)->Struct;
-			if (GenerateWrappedStructType(HostedStruct, OutGeneratedWrappedTypeReferences, OutDirtyModules, true))
+			if (GenerateWrappedStructType(HostedStruct, OutGeneratedWrappedTypeReferences, OutDirtyModules, EPyTypeGenerationFlags::ForceShouldExport))
 			{
 				// Find the wrapped type for the struct as that's what we'll actually add the operator to (via its meta-data)
-				TSharedPtr<PyGenUtil::FGeneratedWrappedStructType> HostedStructGeneratedWrappedType = StaticCastSharedPtr<PyGenUtil::FGeneratedWrappedStructType>(GeneratedWrappedTypes.FindRef(HostedStruct->GetFName()));
+				TSharedPtr<PyGenUtil::FGeneratedWrappedStructType> HostedStructGeneratedWrappedType = StaticCastSharedPtr<PyGenUtil::FGeneratedWrappedStructType>(GeneratedWrappedTypes.FindRef(PyGenUtil::GetTypeRegistryName(HostedStruct)));
 				check(HostedStructGeneratedWrappedType.IsValid());
 				StaticCastSharedPtr<FPyWrapperStructMetaData>(HostedStructGeneratedWrappedType->MetaData)->OpStacks[(int32)OpSignature.OpType].Funcs.Add(MoveTemp(OpFunc));
 			}
@@ -1012,13 +1152,13 @@ PyTypeObject* FPyWrapperTypeRegistry::GenerateWrappedClassType(const UClass* InC
 				const UClass* HostClass = CastChecked<UClass>(HostType);
 
 				// Ensure that we've generated a finalized Python type for this class since we'll be adding this constant to that type
-				if (!GenerateWrappedClassType(HostClass, OutGeneratedWrappedTypeReferences, OutDirtyModules, true))
+				if (!GenerateWrappedClassType(HostClass, OutGeneratedWrappedTypeReferences, OutDirtyModules, EPyTypeGenerationFlags::ForceShouldExport))
 				{
 					return;
 				}
 
 				// Find the wrapped type for the class as that's what we'll actually add the constant to
-				TSharedPtr<PyGenUtil::FGeneratedWrappedClassType> HostedClassGeneratedWrappedType = StaticCastSharedPtr<PyGenUtil::FGeneratedWrappedClassType>(GeneratedWrappedTypes.FindRef(HostClass->GetFName()));
+				TSharedPtr<PyGenUtil::FGeneratedWrappedClassType> HostedClassGeneratedWrappedType = StaticCastSharedPtr<PyGenUtil::FGeneratedWrappedClassType>(GeneratedWrappedTypes.FindRef(PyGenUtil::GetTypeRegistryName(HostClass)));
 				check(HostedClassGeneratedWrappedType.IsValid());
 
 				// Add the dynamic constants to the struct
@@ -1033,13 +1173,13 @@ PyTypeObject* FPyWrapperTypeRegistry::GenerateWrappedClassType(const UClass* InC
 				const UScriptStruct* HostStruct = CastChecked<UScriptStruct>(HostType);
 
 				// Ensure that we've generated a finalized Python type for this struct since we'll be adding this constant to that type
-				if (!GenerateWrappedStructType(HostStruct, OutGeneratedWrappedTypeReferences, OutDirtyModules, true))
+				if (!GenerateWrappedStructType(HostStruct, OutGeneratedWrappedTypeReferences, OutDirtyModules, EPyTypeGenerationFlags::ForceShouldExport))
 				{
 					return;
 				}
 
 				// Find the wrapped type for the struct as that's what we'll actually add the constant to
-				TSharedPtr<PyGenUtil::FGeneratedWrappedStructType> HostedStructGeneratedWrappedType = StaticCastSharedPtr<PyGenUtil::FGeneratedWrappedStructType>(GeneratedWrappedTypes.FindRef(HostStruct->GetFName()));
+				TSharedPtr<PyGenUtil::FGeneratedWrappedStructType> HostedStructGeneratedWrappedType = StaticCastSharedPtr<PyGenUtil::FGeneratedWrappedStructType>(GeneratedWrappedTypes.FindRef(PyGenUtil::GetTypeRegistryName(HostStruct)));
 				check(HostedStructGeneratedWrappedType.IsValid());
 
 				// Add the dynamic constants to the struct
@@ -1169,14 +1309,14 @@ PyTypeObject* FPyWrapperTypeRegistry::GenerateWrappedClassType(const UClass* InC
 	FString TypeDocString = PyGenUtil::PythonizeTooltip(PyGenUtil::GetFieldTooltip(InClass));
 	if (const UClass* SuperClass = InClass->GetSuperClass())
 	{
-		TSharedPtr<PyGenUtil::FGeneratedWrappedClassType> SuperGeneratedWrappedType = StaticCastSharedPtr<PyGenUtil::FGeneratedWrappedClassType>(GeneratedWrappedTypes.FindRef(SuperClass->GetFName()));
+		TSharedPtr<PyGenUtil::FGeneratedWrappedClassType> SuperGeneratedWrappedType = StaticCastSharedPtr<PyGenUtil::FGeneratedWrappedClassType>(GeneratedWrappedTypes.FindRef(PyGenUtil::GetTypeRegistryName(SuperClass)));
 		if (SuperGeneratedWrappedType.IsValid())
 		{
 			GeneratedWrappedType->PropertyDocs.Append(SuperGeneratedWrappedType->PropertyDocs);
 		}
 	}
 	GeneratedWrappedType->PropertyDocs.Sort(&PyGenUtil::FGeneratedWrappedPropertyDoc::SortPredicate);
-	PyGenUtil::AppendCppSourceInformationDocString(InClass, TypeDocString);
+	PyGenUtil::AppendSourceInformationDocString(InClass, TypeDocString);
 	PyGenUtil::FGeneratedWrappedPropertyDoc::AppendDocString(GeneratedWrappedType->PropertyDocs, TypeDocString);
 	GeneratedWrappedType->TypeDoc = PyGenUtil::TCHARToUTF8Buffer(*TypeDocString);
 
@@ -1201,22 +1341,24 @@ PyTypeObject* FPyWrapperTypeRegistry::GenerateWrappedClassType(const UClass* InC
 
 	if (GeneratedWrappedType->Finalize())
 	{
-		const FName UnrealModuleName = *PyGenUtil::GetFieldModule(InClass);
-		GeneratedWrappedTypesForModule.Add(UnrealModuleName, InClass->GetFName());
-		OutDirtyModules.Add(UnrealModuleName);
-
-		const FString PyModuleName = PyGenUtil::GetModulePythonName(UnrealModuleName);
 		PyObject* PyModule = nullptr;
-		// Execute Python code within this block
+		const FName UnrealModuleName = *PyGenUtil::GetFieldModule(InClass);
+		if (!UnrealModuleName.IsNone())
 		{
-			FPyScopedGIL GIL;
+			GeneratedWrappedTypesForModule.Add(UnrealModuleName, TypeRegistryName);
+			OutDirtyModules.Add(UnrealModuleName);
 
-			PyModule = PyImport_AddModule(TCHAR_TO_UTF8(*PyModuleName));
+			const FString PyModuleName = PyGenUtil::GetModulePythonName(UnrealModuleName);
+			// Execute Python code within this block
+			{
+				FPyScopedGIL GIL;
+				PyModule = PyImport_AddModule(TCHAR_TO_UTF8(*PyModuleName));
 
-			Py_INCREF(&GeneratedWrappedType->PyType);
-			PyModule_AddObject(PyModule, GeneratedWrappedType->PyType.tp_name, (PyObject*)&GeneratedWrappedType->PyType);
+				Py_INCREF(&GeneratedWrappedType->PyType);
+				PyModule_AddObject(PyModule, GeneratedWrappedType->PyType.tp_name, (PyObject*)&GeneratedWrappedType->PyType);
+			}
 		}
-		RegisterWrappedClassType(InClass->GetFName(), &GeneratedWrappedType->PyType);
+		RegisterWrappedClassType(TypeRegistryName, &GeneratedWrappedType->PyType, !bIsBlueprintGeneratedType);
 
 		// Also generate and register any deprecated aliases for this type
 		const TArray<FString> DeprecatedPythonClassNames = PyGenUtil::GetDeprecatedClassPythonNames(InClass);
@@ -1246,15 +1388,19 @@ PyTypeObject* FPyWrapperTypeRegistry::GenerateWrappedClassType(const UClass* InC
 
 			if (DeprecatedGeneratedWrappedType->Finalize())
 			{
-				GeneratedWrappedTypesForModule.Add(UnrealModuleName, DeprecatedClassName);
-				// Execute Python code within this block
+				if (!UnrealModuleName.IsNone())
 				{
-					FPyScopedGIL GIL;
+					GeneratedWrappedTypesForModule.Add(UnrealModuleName, DeprecatedClassName);
+					// Execute Python code within this block
+					{
+						FPyScopedGIL GIL;
+						check(PyModule);
 
-					Py_INCREF(&DeprecatedGeneratedWrappedType->PyType);
-					PyModule_AddObject(PyModule, DeprecatedGeneratedWrappedType->PyType.tp_name, (PyObject*)&DeprecatedGeneratedWrappedType->PyType);
+						Py_INCREF(&DeprecatedGeneratedWrappedType->PyType);
+						PyModule_AddObject(PyModule, DeprecatedGeneratedWrappedType->PyType.tp_name, (PyObject*)&DeprecatedGeneratedWrappedType->PyType);
+					}
 				}
-				RegisterWrappedClassType(DeprecatedClassName, &DeprecatedGeneratedWrappedType->PyType);
+				RegisterWrappedClassType(DeprecatedClassName, &DeprecatedGeneratedWrappedType->PyType, !bIsBlueprintGeneratedType);
 			}
 			else
 			{
@@ -1286,7 +1432,8 @@ void FPyWrapperTypeRegistry::UnregisterWrappedClassType(const FName ClassName, P
 
 bool FPyWrapperTypeRegistry::HasWrappedClassType(const UClass* InClass) const
 {
-	return PythonWrappedClasses.Contains(InClass->GetFName());
+	const FName TypeRegistryName = PyGenUtil::GetTypeRegistryName(InClass);
+	return PythonWrappedClasses.Contains(TypeRegistryName);
 }
 
 PyTypeObject* FPyWrapperTypeRegistry::GetWrappedClassType(const UClass* InClass) const
@@ -1295,7 +1442,8 @@ PyTypeObject* FPyWrapperTypeRegistry::GetWrappedClassType(const UClass* InClass)
 
 	for (const UClass* Class = InClass; Class; Class = Class->GetSuperClass())
 	{
-		if (PyTypeObject* ClassPyType = PythonWrappedClasses.FindRef(Class->GetFName()))
+		const FName TypeRegistryName = PyGenUtil::GetTypeRegistryName(Class);
+		if (PyTypeObject* ClassPyType = PythonWrappedClasses.FindRef(TypeRegistryName))
 		{
 			PyType = ClassPyType;
 			break;
@@ -1305,7 +1453,7 @@ PyTypeObject* FPyWrapperTypeRegistry::GetWrappedClassType(const UClass* InClass)
 	return PyType;
 }
 
-PyTypeObject* FPyWrapperTypeRegistry::GenerateWrappedStructType(const UScriptStruct* InStruct, FGeneratedWrappedTypeReferences& OutGeneratedWrappedTypeReferences, TSet<FName>& OutDirtyModules, const bool bForceGenerate)
+PyTypeObject* FPyWrapperTypeRegistry::GenerateWrappedStructType(const UScriptStruct* InStruct, FGeneratedWrappedTypeReferences& OutGeneratedWrappedTypeReferences, TSet<FName>& OutDirtyModules, const EPyTypeGenerationFlags InGenerationFlags)
 {
 	SCOPE_SECONDS_ACCUMULATOR(STAT_GenerateWrappedStructTotalTime);
 	INC_DWORD_STAT(STAT_GenerateWrappedStructCallCount);
@@ -1328,19 +1476,21 @@ PyTypeObject* FPyWrapperTypeRegistry::GenerateWrappedStructType(const UScriptStr
 	};
 
 	// Already processed? Nothing more to do
-	if (PyTypeObject* ExistingPyType = PythonWrappedStructs.FindRef(InStruct->GetFName()))
+	const FName TypeRegistryName = PyGenUtil::GetTypeRegistryName(InStruct);
+	if (PyTypeObject* ExistingPyType = PythonWrappedStructs.FindRef(TypeRegistryName))
 	{
 		return ExistingPyType;
 	}
 
 	// todo: allow generation of Blueprint generated structs
-	if (PyGenUtil::IsBlueprintGeneratedStruct(InStruct))
+	const bool bIsBlueprintGeneratedType = PyGenUtil::IsBlueprintGeneratedStruct(InStruct);
+	if (bIsBlueprintGeneratedType)
 	{
 		return nullptr;
 	}
 
 	// Should this type be exported?
-	if (!bForceGenerate && !PyGenUtil::ShouldExportStruct(InStruct))
+	if (!EnumHasAnyFlags(InGenerationFlags, EPyTypeGenerationFlags::ForceShouldExport) && !PyGenUtil::ShouldExportStruct(InStruct))
 	{
 		return nullptr;
 	}
@@ -1349,14 +1499,14 @@ PyTypeObject* FPyWrapperTypeRegistry::GenerateWrappedStructType(const UScriptStr
 	PyTypeObject* SuperPyType = nullptr;
 	if (const UScriptStruct* SuperStruct = Cast<UScriptStruct>(InStruct->GetSuperStruct()))
 	{
-		SuperPyType = GenerateWrappedStructType(SuperStruct, OutGeneratedWrappedTypeReferences, OutDirtyModules, true);
+		SuperPyType = GenerateWrappedStructType(SuperStruct, OutGeneratedWrappedTypeReferences, OutDirtyModules, InGenerationFlags | EPyTypeGenerationFlags::ForceShouldExport);
 	}
 
 	INC_DWORD_STAT(STAT_GenerateWrappedStructObjCount);
 
-	check(!GeneratedWrappedTypes.Contains(InStruct->GetFName()));
+	check(!GeneratedWrappedTypes.Contains(TypeRegistryName));
 	TSharedRef<PyGenUtil::FGeneratedWrappedStructType> GeneratedWrappedType = MakeShared<PyGenUtil::FGeneratedWrappedStructType>();
-	GeneratedWrappedTypes.Add(InStruct->GetFName(), GeneratedWrappedType);
+	GeneratedWrappedTypes.Add(TypeRegistryName, GeneratedWrappedType);
 
 	TMap<FName, FName> PythonProperties;
 	TMap<FName, FString> PythonDeprecatedProperties;
@@ -1425,19 +1575,19 @@ PyTypeObject* FPyWrapperTypeRegistry::GenerateWrappedStructType(const UScriptStr
 	FString TypeDocString = PyGenUtil::PythonizeTooltip(PyGenUtil::GetFieldTooltip(InStruct));
 	if (const UScriptStruct* SuperStruct = Cast<UScriptStruct>(InStruct->GetSuperStruct()))
 	{
-		TSharedPtr<PyGenUtil::FGeneratedWrappedStructType> SuperGeneratedWrappedType = StaticCastSharedPtr<PyGenUtil::FGeneratedWrappedStructType>(GeneratedWrappedTypes.FindRef(SuperStruct->GetFName()));
+		TSharedPtr<PyGenUtil::FGeneratedWrappedStructType> SuperGeneratedWrappedType = StaticCastSharedPtr<PyGenUtil::FGeneratedWrappedStructType>(GeneratedWrappedTypes.FindRef(PyGenUtil::GetTypeRegistryName(SuperStruct)));
 		if (SuperGeneratedWrappedType.IsValid())
 		{
 			GeneratedWrappedType->PropertyDocs.Append(SuperGeneratedWrappedType->PropertyDocs);
 		}
 	}
 	GeneratedWrappedType->PropertyDocs.Sort(&PyGenUtil::FGeneratedWrappedPropertyDoc::SortPredicate);
-	PyGenUtil::AppendCppSourceInformationDocString(InStruct, TypeDocString);
+	PyGenUtil::AppendSourceInformationDocString(InStruct, TypeDocString);
 	PyGenUtil::FGeneratedWrappedPropertyDoc::AppendDocString(GeneratedWrappedType->PropertyDocs, TypeDocString);
 	GeneratedWrappedType->TypeDoc = PyGenUtil::TCHARToUTF8Buffer(*TypeDocString);
 
 	int32 WrappedStructSizeBytes = sizeof(FPyWrapperStruct);
-	if (const IPyWrapperInlineStructFactory* InlineStructFactory = GetInlineStructFactory(InStruct->GetFName()))
+	if (const IPyWrapperInlineStructFactory* InlineStructFactory = GetInlineStructFactory(TypeRegistryName))
 	{
 		WrappedStructSizeBytes = InlineStructFactory->GetPythonObjectSizeBytes();
 	}
@@ -1531,22 +1681,24 @@ PyTypeObject* FPyWrapperTypeRegistry::GenerateWrappedStructType(const UScriptStr
 
 	if (GeneratedWrappedType->Finalize())
 	{
-		const FName UnrealModuleName = *PyGenUtil::GetFieldModule(InStruct);
-		GeneratedWrappedTypesForModule.Add(UnrealModuleName, InStruct->GetFName());
-		OutDirtyModules.Add(UnrealModuleName);
-
-		const FString PyModuleName = PyGenUtil::GetModulePythonName(UnrealModuleName);
 		PyObject* PyModule = nullptr;
-		// Execute Python code within this block
+		const FName UnrealModuleName = *PyGenUtil::GetFieldModule(InStruct);
+		if (!UnrealModuleName.IsNone())
 		{
-			FPyScopedGIL GIL;
+			GeneratedWrappedTypesForModule.Add(UnrealModuleName, TypeRegistryName);
+			OutDirtyModules.Add(UnrealModuleName);
 
-			PyModule = PyImport_AddModule(TCHAR_TO_UTF8(*PyModuleName));
+			const FString PyModuleName = PyGenUtil::GetModulePythonName(UnrealModuleName);
+			// Execute Python code within this block
+			{
+				FPyScopedGIL GIL;
+				PyModule = PyImport_AddModule(TCHAR_TO_UTF8(*PyModuleName));
 
-			Py_INCREF(&GeneratedWrappedType->PyType);
-			PyModule_AddObject(PyModule, GeneratedWrappedType->PyType.tp_name, (PyObject*)&GeneratedWrappedType->PyType);
+				Py_INCREF(&GeneratedWrappedType->PyType);
+				PyModule_AddObject(PyModule, GeneratedWrappedType->PyType.tp_name, (PyObject*)&GeneratedWrappedType->PyType);
+			}
 		}
-		RegisterWrappedStructType(InStruct->GetFName(), &GeneratedWrappedType->PyType);
+		RegisterWrappedStructType(TypeRegistryName, &GeneratedWrappedType->PyType, !bIsBlueprintGeneratedType);
 
 		// Also generate and register any deprecated aliases for this type
 		const TArray<FString> DeprecatedPythonStructNames = PyGenUtil::GetDeprecatedStructPythonNames(InStruct);
@@ -1576,15 +1728,19 @@ PyTypeObject* FPyWrapperTypeRegistry::GenerateWrappedStructType(const UScriptStr
 
 			if (DeprecatedGeneratedWrappedType->Finalize())
 			{
-				GeneratedWrappedTypesForModule.Add(UnrealModuleName, DeprecatedStructName);
-				// Execute Python code within this block
+				if (!UnrealModuleName.IsNone())
 				{
-					FPyScopedGIL GIL;
+					GeneratedWrappedTypesForModule.Add(UnrealModuleName, DeprecatedStructName);
+					// Execute Python code within this block
+					{
+						FPyScopedGIL GIL;
+						check(PyModule);
 
-					Py_INCREF(&DeprecatedGeneratedWrappedType->PyType);
-					PyModule_AddObject(PyModule, DeprecatedGeneratedWrappedType->PyType.tp_name, (PyObject*)&DeprecatedGeneratedWrappedType->PyType);
+						Py_INCREF(&DeprecatedGeneratedWrappedType->PyType);
+						PyModule_AddObject(PyModule, DeprecatedGeneratedWrappedType->PyType.tp_name, (PyObject*)&DeprecatedGeneratedWrappedType->PyType);
+					}
 				}
-				RegisterWrappedStructType(DeprecatedStructName, &DeprecatedGeneratedWrappedType->PyType);
+				RegisterWrappedStructType(DeprecatedStructName, &DeprecatedGeneratedWrappedType->PyType, !bIsBlueprintGeneratedType);
 			}
 			else
 			{
@@ -1616,7 +1772,8 @@ void FPyWrapperTypeRegistry::UnregisterWrappedStructType(const FName StructName,
 
 bool FPyWrapperTypeRegistry::HasWrappedStructType(const UScriptStruct* InStruct) const
 {
-	return PythonWrappedStructs.Contains(InStruct->GetFName());
+	const FName TypeRegistryName = PyGenUtil::GetTypeRegistryName(InStruct);
+	return PythonWrappedStructs.Contains(TypeRegistryName);
 }
 
 PyTypeObject* FPyWrapperTypeRegistry::GetWrappedStructType(const UScriptStruct* InStruct) const
@@ -1625,7 +1782,8 @@ PyTypeObject* FPyWrapperTypeRegistry::GetWrappedStructType(const UScriptStruct*
 
 	for (const UScriptStruct* Struct = InStruct; Struct; Struct = Cast<UScriptStruct>(Struct->GetSuperStruct()))
 	{
-		if (PyTypeObject* StructPyType = PythonWrappedStructs.FindRef(Struct->GetFName()))
+		const FName TypeRegistryName = PyGenUtil::GetTypeRegistryName(Struct);
+		if (PyTypeObject* StructPyType = PythonWrappedStructs.FindRef(TypeRegistryName))
 		{
 			PyType = StructPyType;
 			break;
@@ -1635,37 +1793,54 @@ PyTypeObject* FPyWrapperTypeRegistry::GetWrappedStructType(const UScriptStruct*
 	return PyType;
 }
 
-PyTypeObject* FPyWrapperTypeRegistry::GenerateWrappedEnumType(const UEnum* InEnum, FGeneratedWrappedTypeReferences& OutGeneratedWrappedTypeReferences, TSet<FName>& OutDirtyModules, const bool bForceGenerate)
+PyTypeObject* FPyWrapperTypeRegistry::GenerateWrappedEnumType(const UEnum* InEnum, FGeneratedWrappedTypeReferences& OutGeneratedWrappedTypeReferences, TSet<FName>& OutDirtyModules, const EPyTypeGenerationFlags InGenerationFlags)
 {
 	SCOPE_SECONDS_ACCUMULATOR(STAT_GenerateWrappedEnumTotalTime);
 	INC_DWORD_STAT(STAT_GenerateWrappedEnumCallCount);
 
+	bool bIsNewType = true;
+
 	// Already processed? Nothing more to do
-	if (PyTypeObject* ExistingPyType = PythonWrappedEnums.FindRef(InEnum->GetFName()))
+	const FName TypeRegistryName = PyGenUtil::GetTypeRegistryName(InEnum);
+	if (PyTypeObject* ExistingPyType = PythonWrappedEnums.FindRef(TypeRegistryName))
 	{
-		return ExistingPyType;
+		if (!EnumHasAnyFlags(InGenerationFlags, EPyTypeGenerationFlags::OverwriteExisting))
+		{
+			return ExistingPyType;
+		}
+		bIsNewType = false;
 	}
 
-	// todo: allow generation of Blueprint generated enums
-	if (PyGenUtil::IsBlueprintGeneratedEnum(InEnum))
+	// Should we allow Blueprint generated enums?
+	const bool bIsBlueprintGeneratedType = PyGenUtil::IsBlueprintGeneratedEnum(InEnum);
+	if (!EnumHasAnyFlags(InGenerationFlags, EPyTypeGenerationFlags::IncludeBlueprintGeneratedTypes) && bIsBlueprintGeneratedType)
 	{
 		return nullptr;
 	}
 
 	// Should this type be exported?
-	if (!bForceGenerate && !PyGenUtil::ShouldExportEnum(InEnum))
+	if (!EnumHasAnyFlags(InGenerationFlags, EPyTypeGenerationFlags::ForceShouldExport) && !PyGenUtil::ShouldExportEnum(InEnum))
 	{
 		return nullptr;
 	}
 
 	INC_DWORD_STAT(STAT_GenerateWrappedEnumObjCount);
 
-	check(!GeneratedWrappedTypes.Contains(InEnum->GetFName()));
-	TSharedRef<PyGenUtil::FGeneratedWrappedEnumType> GeneratedWrappedType = MakeShared<PyGenUtil::FGeneratedWrappedEnumType>();
-	GeneratedWrappedTypes.Add(InEnum->GetFName(), GeneratedWrappedType);
+	TSharedPtr<PyGenUtil::FGeneratedWrappedEnumType> GeneratedWrappedType;
+	if (bIsNewType)
+	{
+		check(!GeneratedWrappedTypes.Contains(TypeRegistryName));
+		GeneratedWrappedType = MakeShared<PyGenUtil::FGeneratedWrappedEnumType>();
+		GeneratedWrappedTypes.Add(TypeRegistryName, GeneratedWrappedType);
+	}
+	else
+	{
+		GeneratedWrappedType = StaticCastSharedPtr<PyGenUtil::FGeneratedWrappedEnumType>(GeneratedWrappedTypes.FindChecked(TypeRegistryName));
+		GeneratedWrappedType->Reset();
+	}
 
 	FString TypeDocString = PyGenUtil::PythonizeTooltip(PyGenUtil::GetFieldTooltip(InEnum));
-	PyGenUtil::AppendCppSourceInformationDocString(InEnum, TypeDocString);
+	PyGenUtil::AppendSourceInformationDocString(InEnum, TypeDocString);
 
 	const FString PythonEnumName = PyGenUtil::GetEnumPythonName(InEnum);
 	GeneratedWrappedType->TypeName = PyGenUtil::TCHARToUTF8Buffer(*PythonEnumName);
@@ -1682,65 +1857,74 @@ PyTypeObject* FPyWrapperTypeRegistry::GenerateWrappedEnumType(const UEnum* InEnu
 
 	if (GeneratedWrappedType->Finalize())
 	{
-		const FName UnrealModuleName = *PyGenUtil::GetFieldModule(InEnum);
-		GeneratedWrappedTypesForModule.Add(UnrealModuleName, InEnum->GetFName());
-		OutDirtyModules.Add(UnrealModuleName);
-
-		const FString PyModuleName = PyGenUtil::GetModulePythonName(UnrealModuleName);
-		PyObject* PyModule = nullptr;
-		// Execute Python code within this block
+		if (bIsNewType)
 		{
-			FPyScopedGIL GIL;
-
-			PyModule = PyImport_AddModule(TCHAR_TO_UTF8(*PyModuleName));
+			PyObject* PyModule = nullptr;
+			const FName UnrealModuleName = *PyGenUtil::GetFieldModule(InEnum);
+			if (!UnrealModuleName.IsNone())
+			{
+				GeneratedWrappedTypesForModule.Add(UnrealModuleName, TypeRegistryName);
+				OutDirtyModules.Add(UnrealModuleName);
 
-			Py_INCREF(&GeneratedWrappedType->PyType);
-			PyModule_AddObject(PyModule, GeneratedWrappedType->PyType.tp_name, (PyObject*)&GeneratedWrappedType->PyType);
-		}
-		RegisterWrappedEnumType(InEnum->GetFName(), &GeneratedWrappedType->PyType);
+				const FString PyModuleName = PyGenUtil::GetModulePythonName(UnrealModuleName);
+				// Execute Python code within this block
+				{
+					FPyScopedGIL GIL;
+					PyModule = PyImport_AddModule(TCHAR_TO_UTF8(*PyModuleName));
 
-		// Also generate and register any deprecated aliases for this type
-		const TArray<FString> DeprecatedPythonEnumNames = PyGenUtil::GetDeprecatedEnumPythonNames(InEnum);
-		for (const FString& DeprecatedPythonEnumName : DeprecatedPythonEnumNames)
-		{
-			const FName DeprecatedEnumName = *DeprecatedPythonEnumName;
-			FString DeprecationMessage = FString::Printf(TEXT("'%s' was renamed to '%s'."), *DeprecatedPythonEnumName, *PythonEnumName);
+					Py_INCREF(&GeneratedWrappedType->PyType);
+					PyModule_AddObject(PyModule, GeneratedWrappedType->PyType.tp_name, (PyObject*)&GeneratedWrappedType->PyType);
+				}
+			}
+			RegisterWrappedEnumType(TypeRegistryName, &GeneratedWrappedType->PyType, !bIsBlueprintGeneratedType);
 
-			if (GeneratedWrappedTypes.Contains(DeprecatedEnumName))
+			// Also generate and register any deprecated aliases for this type
+			const TArray<FString> DeprecatedPythonEnumNames = PyGenUtil::GetDeprecatedEnumPythonNames(InEnum);
+			for (const FString& DeprecatedPythonEnumName : DeprecatedPythonEnumNames)
 			{
-				REPORT_PYTHON_GENERATION_ISSUE(Warning, TEXT("Deprecated enum name '%s' conflicted with an existing type!"), *DeprecatedPythonEnumName);
-				continue;
-			}
+				const FName DeprecatedEnumName = *DeprecatedPythonEnumName;
+				FString DeprecationMessage = FString::Printf(TEXT("'%s' was renamed to '%s'."), *DeprecatedPythonEnumName, *PythonEnumName);
 
-			TSharedRef<PyGenUtil::FGeneratedWrappedEnumType> DeprecatedGeneratedWrappedType = MakeShared<PyGenUtil::FGeneratedWrappedEnumType>();
-			GeneratedWrappedTypes.Add(DeprecatedEnumName, DeprecatedGeneratedWrappedType);
+				if (GeneratedWrappedTypes.Contains(DeprecatedEnumName))
+				{
+					REPORT_PYTHON_GENERATION_ISSUE(Warning, TEXT("Deprecated enum name '%s' conflicted with an existing type!"), *DeprecatedPythonEnumName);
+					continue;
+				}
 
-			DeprecatedGeneratedWrappedType->TypeName = PyGenUtil::TCHARToUTF8Buffer(*DeprecatedPythonEnumName);
-			DeprecatedGeneratedWrappedType->TypeDoc = PyGenUtil::TCHARToUTF8Buffer(*FString::Printf(TEXT("deprecated: %s"), *DeprecationMessage));
-			DeprecatedGeneratedWrappedType->EnumEntries = GeneratedWrappedType->EnumEntries;
-			DeprecatedGeneratedWrappedType->PyType.tp_basicsize = sizeof(FPyWrapperEnum);
-			DeprecatedGeneratedWrappedType->PyType.tp_base = &PyWrapperEnumType;
-			DeprecatedGeneratedWrappedType->PyType.tp_flags = Py_TPFLAGS_DEFAULT;
+				TSharedRef<PyGenUtil::FGeneratedWrappedEnumType> DeprecatedGeneratedWrappedType = MakeShared<PyGenUtil::FGeneratedWrappedEnumType>();
+				GeneratedWrappedTypes.Add(DeprecatedEnumName, DeprecatedGeneratedWrappedType);
 
-			TSharedRef<FPyWrapperEnumMetaData> DeprecatedEnumMetaData = MakeShared<FPyWrapperEnumMetaData>(*EnumMetaData);
-			DeprecatedEnumMetaData->DeprecationMessage = MoveTemp(DeprecationMessage);
-			DeprecatedGeneratedWrappedType->MetaData = DeprecatedEnumMetaData;
+				DeprecatedGeneratedWrappedType->TypeName = PyGenUtil::TCHARToUTF8Buffer(*DeprecatedPythonEnumName);
+				DeprecatedGeneratedWrappedType->TypeDoc = PyGenUtil::TCHARToUTF8Buffer(*FString::Printf(TEXT("deprecated: %s"), *DeprecationMessage));
+				DeprecatedGeneratedWrappedType->EnumEntries = GeneratedWrappedType->EnumEntries;
+				DeprecatedGeneratedWrappedType->PyType.tp_basicsize = sizeof(FPyWrapperEnum);
+				DeprecatedGeneratedWrappedType->PyType.tp_base = &PyWrapperEnumType;
+				DeprecatedGeneratedWrappedType->PyType.tp_flags = Py_TPFLAGS_DEFAULT;
 
-			if (DeprecatedGeneratedWrappedType->Finalize())
-			{
-				GeneratedWrappedTypesForModule.Add(UnrealModuleName, DeprecatedEnumName);
-				// Execute Python code within this block
+				TSharedRef<FPyWrapperEnumMetaData> DeprecatedEnumMetaData = MakeShared<FPyWrapperEnumMetaData>(*EnumMetaData);
+				DeprecatedEnumMetaData->DeprecationMessage = MoveTemp(DeprecationMessage);
+				DeprecatedGeneratedWrappedType->MetaData = DeprecatedEnumMetaData;
+
+				if (DeprecatedGeneratedWrappedType->Finalize())
 				{
-					FPyScopedGIL GIL;
+					if (!UnrealModuleName.IsNone())
+					{
+						GeneratedWrappedTypesForModule.Add(UnrealModuleName, DeprecatedEnumName);
+						// Execute Python code within this block
+						{
+							FPyScopedGIL GIL;
+							check(PyModule);
 
-					Py_INCREF(&DeprecatedGeneratedWrappedType->PyType);
-					PyModule_AddObject(PyModule, DeprecatedGeneratedWrappedType->PyType.tp_name, (PyObject*)&DeprecatedGeneratedWrappedType->PyType);
+							Py_INCREF(&DeprecatedGeneratedWrappedType->PyType);
+							PyModule_AddObject(PyModule, DeprecatedGeneratedWrappedType->PyType.tp_name, (PyObject*)&DeprecatedGeneratedWrappedType->PyType);
+						}
+					}
+					RegisterWrappedEnumType(DeprecatedEnumName, &DeprecatedGeneratedWrappedType->PyType, !bIsBlueprintGeneratedType);
+				}
+				else
+				{
+					REPORT_PYTHON_GENERATION_ISSUE(Fatal, TEXT("Failed to generate Python glue code for deprecated enum '%s'!"), *DeprecatedPythonEnumName);
 				}
-				RegisterWrappedEnumType(DeprecatedEnumName, &DeprecatedGeneratedWrappedType->PyType);
-			}
-			else
-			{
-				REPORT_PYTHON_GENERATION_ISSUE(Fatal, TEXT("Failed to generate Python glue code for deprecated enum '%s'!"), *DeprecatedPythonEnumName);
 			}
 		}
 
@@ -1768,14 +1952,16 @@ void FPyWrapperTypeRegistry::UnregisterWrappedEnumType(const FName EnumName, PyT
 
 bool FPyWrapperTypeRegistry::HasWrappedEnumType(const UEnum* InEnum) const
 {
-	return PythonWrappedEnums.Contains(InEnum->GetFName());
+	const FName TypeRegistryName = PyGenUtil::GetTypeRegistryName(InEnum);
+	return PythonWrappedEnums.Contains(TypeRegistryName);
 }
 
 PyTypeObject* FPyWrapperTypeRegistry::GetWrappedEnumType(const UEnum* InEnum) const
 {
 	PyTypeObject* PyType = &PyWrapperEnumType;
 
-	if (PyTypeObject* EnumPyType = PythonWrappedEnums.FindRef(InEnum->GetFName()))
+	const FName TypeRegistryName = PyGenUtil::GetTypeRegistryName(InEnum);
+	if (PyTypeObject* EnumPyType = PythonWrappedEnums.FindRef(TypeRegistryName))
 	{
 		PyType = EnumPyType;
 	}
@@ -1783,13 +1969,14 @@ PyTypeObject* FPyWrapperTypeRegistry::GetWrappedEnumType(const UEnum* InEnum) co
 	return PyType;
 }
 
-PyTypeObject* FPyWrapperTypeRegistry::GenerateWrappedDelegateType(const UFunction* InDelegateSignature, FGeneratedWrappedTypeReferences& OutGeneratedWrappedTypeReferences, TSet<FName>& OutDirtyModules, const bool bForceGenerate)
+PyTypeObject* FPyWrapperTypeRegistry::GenerateWrappedDelegateType(const UFunction* InDelegateSignature, FGeneratedWrappedTypeReferences& OutGeneratedWrappedTypeReferences, TSet<FName>& OutDirtyModules, const EPyTypeGenerationFlags InGenerationFlags)
 {
 	SCOPE_SECONDS_ACCUMULATOR(STAT_GenerateWrappedDelegateTotalTime);
 	INC_DWORD_STAT(STAT_GenerateWrappedDelegateCallCount);
 
 	// Already processed? Nothing more to do
-	if (PyTypeObject* ExistingPyType = PythonWrappedDelegates.FindRef(InDelegateSignature->GetFName()))
+	const FName TypeRegistryName = PyGenUtil::GetTypeRegistryName(InDelegateSignature);
+	if (PyTypeObject* ExistingPyType = PythonWrappedDelegates.FindRef(TypeRegistryName))
 	{
 		return ExistingPyType;
 	}
@@ -1802,9 +1989,9 @@ PyTypeObject* FPyWrapperTypeRegistry::GenerateWrappedDelegateType(const UFunctio
 
 	INC_DWORD_STAT(STAT_GenerateWrappedDelegateObjCount);
 
-	check(!GeneratedWrappedTypes.Contains(InDelegateSignature->GetFName()));
+	check(!GeneratedWrappedTypes.Contains(TypeRegistryName));
 	TSharedRef<PyGenUtil::FGeneratedWrappedType> GeneratedWrappedType = MakeShared<PyGenUtil::FGeneratedWrappedType>();
-	GeneratedWrappedTypes.Add(InDelegateSignature->GetFName(), GeneratedWrappedType);
+	GeneratedWrappedTypes.Add(TypeRegistryName, GeneratedWrappedType);
 
 	for (TFieldIterator<const UProperty> ParamIt(InDelegateSignature); ParamIt; ++ParamIt)
 	{
@@ -1813,7 +2000,7 @@ PyTypeObject* FPyWrapperTypeRegistry::GenerateWrappedDelegateType(const UFunctio
 	}
 
 	FString TypeDocString = PyGenUtil::PythonizeFunctionTooltip(PyGenUtil::GetFieldTooltip(InDelegateSignature), InDelegateSignature);
-	PyGenUtil::AppendCppSourceInformationDocString(InDelegateSignature, TypeDocString);
+	PyGenUtil::AppendSourceInformationDocString(InDelegateSignature, TypeDocString);
 
 	const FString DelegateBaseTypename = PyGenUtil::GetDelegatePythonName(InDelegateSignature);
 	GeneratedWrappedType->TypeName = PyGenUtil::TCHARToUTF8Buffer(*DelegateBaseTypename);
@@ -1867,20 +2054,22 @@ PyTypeObject* FPyWrapperTypeRegistry::GenerateWrappedDelegateType(const UFunctio
 	if (GeneratedWrappedType->Finalize())
 	{
 		const FName UnrealModuleName = *PyGenUtil::GetFieldModule(InDelegateSignature);
-		GeneratedWrappedTypesForModule.Add(UnrealModuleName, InDelegateSignature->GetFName());
-		OutDirtyModules.Add(UnrealModuleName);
-
-		const FString PyModuleName = PyGenUtil::GetModulePythonName(UnrealModuleName);
-		// Execute Python code within this block
+		if (!UnrealModuleName.IsNone())
 		{
-			FPyScopedGIL GIL;
+			GeneratedWrappedTypesForModule.Add(UnrealModuleName, TypeRegistryName);
+			OutDirtyModules.Add(UnrealModuleName);
 
-			PyObject* PyModule = PyImport_AddModule(TCHAR_TO_UTF8(*PyModuleName));
+			const FString PyModuleName = PyGenUtil::GetModulePythonName(UnrealModuleName);
+			// Execute Python code within this block
+			{
+				FPyScopedGIL GIL;
+				PyObject* PyModule = PyImport_AddModule(TCHAR_TO_UTF8(*PyModuleName));
 
-			Py_INCREF(&GeneratedWrappedType->PyType);
-			PyModule_AddObject(PyModule, GeneratedWrappedType->PyType.tp_name, (PyObject*)&GeneratedWrappedType->PyType);
+				Py_INCREF(&GeneratedWrappedType->PyType);
+				PyModule_AddObject(PyModule, GeneratedWrappedType->PyType.tp_name, (PyObject*)&GeneratedWrappedType->PyType);
+			}
 		}
-		RegisterWrappedDelegateType(InDelegateSignature->GetFName(), &GeneratedWrappedType->PyType);
+		RegisterWrappedDelegateType(TypeRegistryName, &GeneratedWrappedType->PyType);
 
 		return &GeneratedWrappedType->PyType;
 	}
@@ -1906,14 +2095,16 @@ void FPyWrapperTypeRegistry::UnregisterWrappedDelegateType(const FName DelegateN
 
 bool FPyWrapperTypeRegistry::HasWrappedDelegateType(const UFunction* InDelegateSignature) const
 {
-	return PythonWrappedDelegates.Contains(InDelegateSignature->GetFName());
+	const FName TypeRegistryName = PyGenUtil::GetTypeRegistryName(InDelegateSignature);
+	return PythonWrappedDelegates.Contains(TypeRegistryName);
 }
 
 PyTypeObject* FPyWrapperTypeRegistry::GetWrappedDelegateType(const UFunction* InDelegateSignature) const
 {
 	PyTypeObject* PyType = InDelegateSignature->HasAnyFunctionFlags(FUNC_MulticastDelegate) ? &PyWrapperMulticastDelegateType : &PyWrapperDelegateType;
 
-	if (PyTypeObject* DelegatePyType = PythonWrappedDelegates.FindRef(InDelegateSignature->GetFName()))
+	const FName TypeRegistryName = PyGenUtil::GetTypeRegistryName(InDelegateSignature);
+	if (PyTypeObject* DelegatePyType = PythonWrappedDelegates.FindRef(TypeRegistryName))
 	{
 		PyType = DelegatePyType;
 	}
@@ -1925,7 +2116,7 @@ void FPyWrapperTypeRegistry::GatherWrappedTypesForPropertyReferences(const UProp
 {
 	if (const UObjectProperty* ObjProp = Cast<const UObjectProperty>(InProp))
 	{
-		if (ObjProp->PropertyClass && !PythonWrappedClasses.Contains(ObjProp->PropertyClass->GetFName()))
+		if (ObjProp->PropertyClass && !PythonWrappedClasses.Contains(PyGenUtil::GetTypeRegistryName(ObjProp->PropertyClass)))
 		{
 			OutGeneratedWrappedTypeReferences.ClassReferences.Add(ObjProp->PropertyClass);
 		}
@@ -1934,7 +2125,7 @@ void FPyWrapperTypeRegistry::GatherWrappedTypesForPropertyReferences(const UProp
 
 	if (const UStructProperty* StructProp = Cast<const UStructProperty>(InProp))
 	{
-		if (!PythonWrappedStructs.Contains(StructProp->Struct->GetFName()))
+		if (!PythonWrappedStructs.Contains(PyGenUtil::GetTypeRegistryName(StructProp->Struct)))
 		{
 			OutGeneratedWrappedTypeReferences.StructReferences.Add(StructProp->Struct);
 		}
@@ -1943,7 +2134,7 @@ void FPyWrapperTypeRegistry::GatherWrappedTypesForPropertyReferences(const UProp
 
 	if (const UEnumProperty* EnumProp = Cast<const UEnumProperty>(InProp))
 	{
-		if (!PythonWrappedStructs.Contains(EnumProp->GetEnum()->GetFName()))
+		if (!PythonWrappedStructs.Contains(PyGenUtil::GetTypeRegistryName(EnumProp->GetEnum())))
 		{
 			OutGeneratedWrappedTypeReferences.EnumReferences.Add(EnumProp->GetEnum());
 		}
@@ -1954,7 +2145,7 @@ void FPyWrapperTypeRegistry::GatherWrappedTypesForPropertyReferences(const UProp
 	{
 		if (ByteProp->Enum)
 		{
-			if (!PythonWrappedStructs.Contains(ByteProp->Enum->GetFName()))
+			if (!PythonWrappedStructs.Contains(PyGenUtil::GetTypeRegistryName(ByteProp->Enum)))
 			{
 				OutGeneratedWrappedTypeReferences.EnumReferences.Add(ByteProp->Enum);
 			}
@@ -1964,7 +2155,7 @@ void FPyWrapperTypeRegistry::GatherWrappedTypesForPropertyReferences(const UProp
 
 	if (const UDelegateProperty* DelegateProp = Cast<const UDelegateProperty>(InProp))
 	{
-		if (!PythonWrappedStructs.Contains(DelegateProp->SignatureFunction->GetFName()))
+		if (!PythonWrappedStructs.Contains(PyGenUtil::GetTypeRegistryName(DelegateProp->SignatureFunction)))
 		{
 			OutGeneratedWrappedTypeReferences.DelegateReferences.Add(DelegateProp->SignatureFunction);
 		}
@@ -1973,7 +2164,7 @@ void FPyWrapperTypeRegistry::GatherWrappedTypesForPropertyReferences(const UProp
 
 	if (const UMulticastDelegateProperty* DelegateProp = Cast<const UMulticastDelegateProperty>(InProp))
 	{
-		if (!PythonWrappedStructs.Contains(DelegateProp->SignatureFunction->GetFName()))
+		if (!PythonWrappedStructs.Contains(PyGenUtil::GetTypeRegistryName(DelegateProp->SignatureFunction)))
 		{
 			OutGeneratedWrappedTypeReferences.DelegateReferences.Add(DelegateProp->SignatureFunction);
 		}
diff --git a/Engine/Plugins/Experimental/PythonScriptPlugin/Source/PythonScriptPlugin/Private/PyWrapperTypeRegistry.h b/Engine/Plugins/Experimental/PythonScriptPlugin/Source/PythonScriptPlugin/Private/PyWrapperTypeRegistry.h
--- a/Engine/Plugins/Experimental/PythonScriptPlugin/Source/PythonScriptPlugin/Private/PyWrapperTypeRegistry.h
+++ b/Engine/Plugins/Experimental/PythonScriptPlugin/Source/PythonScriptPlugin/Private/PyWrapperTypeRegistry.h
@@ -9,6 +9,7 @@
 #include "PyGenUtil.h"
 #include "PyConversionMethod.h"
 #include "PyOnlineDocsWriter.h"
+#include "Misc/EnumClassFlags.h"
 #include "UObject/WeakObjectPtr.h"
 #include "Templates/Function.h"
 
@@ -33,6 +34,19 @@ class IPyWrapperInlineStructFactory;
 class FPyFileWriter;
 class FReferenceCollector;
 
+enum class EPyTypeGenerationFlags : uint8
+{
+	/** No behavior */
+	None = 0,
+	/** Generate the Python wrapper for this type, even if it fails the PyGenUtil::ShouldExportX check */
+	ForceShouldExport = 1<<0,
+	/** Generate the Python wrapper for this type, even if it passes the PyGenUtil::IsBlueprintGeneratedX check */
+	IncludeBlueprintGeneratedTypes = 1<<1,
+	/** Generate the Python wrapper for this type, re-using the existing type if it already exists */
+	OverwriteExisting = 1<<2,
+};
+ENUM_CLASS_FLAGS(EPyTypeGenerationFlags);
+
 /** Type conversion for TPyWrapperTypeFactory */
 template <typename UnrealType, typename KeyType>
 struct TPyWrapperTypeFactoryConversion
@@ -368,11 +382,26 @@ public:
 	/** Generate notifications to (re)load the dirtied modules in Python */
 	void NotifyModulesDirtied(const TSet<FName>& InDirtyModules) const;
 
+	/** Called when an asset that is registered with the type registry has been renamed */
+	void UpdateGenerateWrappedTypeForRename(const FName InOldName, const UObject* InObj);
+
+	/** Called when an asset that is registered with the type registry is deleted */
+	void RemoveGenerateWrappedTypeForDelete(const UObject* InObj);
+
 	/** Generate a wrapped type for the given object (if it's a valid type to be wrapped) */
-	PyTypeObject* GenerateWrappedTypeForObject(const UObject* InObj, FGeneratedWrappedTypeReferences& OutGeneratedWrappedTypeReferences, TSet<FName>& OutDirtyModules, const bool bForceGenerate = false);
+	PyTypeObject* GenerateWrappedTypeForObject(const UObject* InObj, FGeneratedWrappedTypeReferences& OutGeneratedWrappedTypeReferences, TSet<FName>& OutDirtyModules, const EPyTypeGenerationFlags InGenerationFlags = EPyTypeGenerationFlags::None);
+
+	/** True if we have wrapped type for the given object */
+	bool HasWrappedTypeForObject(const UObject* InObj) const;
+
+	/** True if we have wrapped type for the given object name */
+	bool HasWrappedTypeForObjectName(const FName InName) const;
+
+	/** Get the best wrapped type for the given object */
+	PyTypeObject* GetWrappedTypeForObject(const UObject* InObj) const;
 
 	/** Generate a wrapped type for the given class (only if this class has not yet been registered; will also register the type) */
-	PyTypeObject* GenerateWrappedClassType(const UClass* InClass, FGeneratedWrappedTypeReferences& OutGeneratedWrappedTypeReferences, TSet<FName>& OutDirtyModules, const bool bForceGenerate = false);
+	PyTypeObject* GenerateWrappedClassType(const UClass* InClass, FGeneratedWrappedTypeReferences& OutGeneratedWrappedTypeReferences, TSet<FName>& OutDirtyModules, const EPyTypeGenerationFlags InGenerationFlags = EPyTypeGenerationFlags::None);
 
 	/** Register the wrapped type associated with the given class name */
 	void RegisterWrappedClassType(const FName ClassName, PyTypeObject* PyType, const bool InDetectNameConflicts = true);
@@ -387,7 +416,7 @@ public:
 	PyTypeObject* GetWrappedClassType(const UClass* InClass) const;
 
 	/** Generate a wrapped type for the given struct (only if this struct has not yet been registered; will also register the type) */
-	PyTypeObject* GenerateWrappedStructType(const UScriptStruct* InStruct, FGeneratedWrappedTypeReferences& OutGeneratedWrappedTypeReferences, TSet<FName>& OutDirtyModules, const bool bForceGenerate = false);
+	PyTypeObject* GenerateWrappedStructType(const UScriptStruct* InStruct, FGeneratedWrappedTypeReferences& OutGeneratedWrappedTypeReferences, TSet<FName>& OutDirtyModules, const EPyTypeGenerationFlags InGenerationFlags = EPyTypeGenerationFlags::None);
 
 	/** Register the wrapped type associated with the given struct name */
 	void RegisterWrappedStructType(const FName StructName, PyTypeObject* PyType, const bool InDetectNameConflicts = true);
@@ -402,7 +431,7 @@ public:
 	PyTypeObject* GetWrappedStructType(const UScriptStruct* InStruct) const;
 
 	/** Generate a wrapped type for the given enum (only if this enum has not yet been registered; will also register the type) */
-	PyTypeObject* GenerateWrappedEnumType(const UEnum* InEnum, FGeneratedWrappedTypeReferences& OutGeneratedWrappedTypeReferences, TSet<FName>& OutDirtyModules, const bool bForceGenerate = false);
+	PyTypeObject* GenerateWrappedEnumType(const UEnum* InEnum, FGeneratedWrappedTypeReferences& OutGeneratedWrappedTypeReferences, TSet<FName>& OutDirtyModules, const EPyTypeGenerationFlags InGenerationFlags = EPyTypeGenerationFlags::None);
 
 	/** Register the wrapped type associated with the given enum name */
 	void RegisterWrappedEnumType(const FName EnumName, PyTypeObject* PyType, const bool InDetectNameConflicts = true);
@@ -417,7 +446,7 @@ public:
 	PyTypeObject* GetWrappedEnumType(const UEnum* InEnum) const;
 
 	/** Generate a wrapped type for the given delegate signature (only if this delegate has not yet been registered; will also register the type) */
-	PyTypeObject* GenerateWrappedDelegateType(const UFunction* InDelegateSignature, FGeneratedWrappedTypeReferences& OutGeneratedWrappedTypeReferences, TSet<FName>& OutDirtyModules, const bool bForceGenerate = false);
+	PyTypeObject* GenerateWrappedDelegateType(const UFunction* InDelegateSignature, FGeneratedWrappedTypeReferences& OutGeneratedWrappedTypeReferences, TSet<FName>& OutDirtyModules, const EPyTypeGenerationFlags InGenerationFlags = EPyTypeGenerationFlags::None);
 
 	/** Register the wrapped type associated with the given delegate name */
 	void RegisterWrappedDelegateType(const FName DelegateName, PyTypeObject* PyType, const bool InDetectNameConflicts = true);
diff --git a/Engine/Plugins/Experimental/PythonScriptPlugin/Source/PythonScriptPlugin/Private/PythonScriptPlugin.cpp b/Engine/Plugins/Experimental/PythonScriptPlugin/Source/PythonScriptPlugin/Private/PythonScriptPlugin.cpp
--- a/Engine/Plugins/Experimental/PythonScriptPlugin/Source/PythonScriptPlugin/Private/PythonScriptPlugin.cpp
+++ b/Engine/Plugins/Experimental/PythonScriptPlugin/Source/PythonScriptPlugin/Private/PythonScriptPlugin.cpp
@@ -15,6 +15,9 @@
 #include "PyWrapperTypeRegistry.h"
 #include "EngineAnalytics.h"
 #include "Interfaces/IAnalyticsProvider.h"
+#include "AssetRegistryModule.h"
+#include "Modules/ModuleManager.h"
+#include "UObject/PackageReload.h"
 #include "Misc/App.h"
 #include "Misc/CoreDelegates.h"
 #include "Misc/Paths.h"
@@ -565,6 +568,15 @@ bool FPythonScriptPlugin::Exec(UWorld* InWorld, const TCHAR* Cmd, FOutputDevice&
 	return false;
 }
 
+void FPythonScriptPlugin::PreChange(const UUserDefinedEnum* Enum, FEnumEditorUtils::EEnumEditorChangeInfo Info)
+{
+}
+
+void FPythonScriptPlugin::PostChange(const UUserDefinedEnum* Enum, FEnumEditorUtils::EEnumEditorChangeInfo Info)
+{
+	OnAssetUpdated(Enum);
+}
+
 #if WITH_PYTHON
 
 void FPythonScriptPlugin::InitializePython()
@@ -684,6 +696,11 @@ void FPythonScriptPlugin::InitializePython()
 
 		FPackageName::OnContentPathMounted().AddRaw(this, &FPythonScriptPlugin::OnContentPathMounted);
 		FPackageName::OnContentPathDismounted().AddRaw(this, &FPythonScriptPlugin::OnContentPathDismounted);
+		FCoreUObjectDelegates::OnPackageReloaded.AddRaw(this, &FPythonScriptPlugin::OnAssetReload);
+
+		FAssetRegistryModule& AssetRegistryModule = FModuleManager::LoadModuleChecked<FAssetRegistryModule>("AssetRegistry");
+		AssetRegistryModule.Get().OnAssetRenamed().AddRaw(this, &FPythonScriptPlugin::OnAssetRenamed);
+		AssetRegistryModule.Get().OnAssetRemoved().AddRaw(this, &FPythonScriptPlugin::OnAssetRemoved);
 	}
 
 	// Initialize the Unreal Python module
@@ -751,6 +768,13 @@ void FPythonScriptPlugin::ShutdownPython()
 
 	FPackageName::OnContentPathMounted().RemoveAll(this);
 	FPackageName::OnContentPathDismounted().RemoveAll(this);
+	FCoreUObjectDelegates::OnPackageReloaded.RemoveAll(this);
+
+	if (FAssetRegistryModule* AssetRegistryModule = FModuleManager::GetModulePtr<FAssetRegistryModule>("AssetRegistry"))
+	{
+		AssetRegistryModule->Get().OnAssetRenamed().RemoveAll(this);
+		AssetRegistryModule->Get().OnAssetRemoved().RemoveAll(this);
+	}
 
 #if WITH_EDITOR
 	FEditorSupportDelegates::PrepareToCleanseEditorObject.RemoveAll(this);
@@ -1145,6 +1169,68 @@ void FPythonScriptPlugin::OnContentPathDismounted(const FString& InAssetPath, co
 	PyUtil::RemoveSystemPath(FPaths::ConvertRelativePathToFull(InFilesystemPath / TEXT("Python")));
 }
 
+void FPythonScriptPlugin::OnAssetRenamed(const FAssetData& Data, const FString& OldName)
+{
+	const FName OldPackageName = *FPackageName::ObjectPathToPackageName(OldName);
+	const UObject* AssetPtr = PyGenUtil::GetTypeRegistryType(Data.GetAsset());
+	if (AssetPtr)
+	{
+		// If this asset has an associated Python type, then we need to rename it
+		FPyWrapperTypeRegistry& PyWrapperTypeRegistry = FPyWrapperTypeRegistry::Get();
+		if (PyWrapperTypeRegistry.HasWrappedTypeForObjectName(OldPackageName))
+		{
+			PyWrapperTypeRegistry.UpdateGenerateWrappedTypeForRename(OldPackageName, AssetPtr);
+			OnAssetUpdated(AssetPtr);
+		}
+	}
+}
+
+void FPythonScriptPlugin::OnAssetRemoved(const FAssetData& Data)
+{
+	const UObject* AssetPtr = PyGenUtil::GetTypeRegistryType(Data.GetAsset());
+	if (AssetPtr)
+	{
+		// If this asset has an associated Python type, then we need to remove it
+		FPyWrapperTypeRegistry& PyWrapperTypeRegistry = FPyWrapperTypeRegistry::Get();
+		if (PyWrapperTypeRegistry.HasWrappedTypeForObject(AssetPtr))
+		{
+			PyWrapperTypeRegistry.RemoveGenerateWrappedTypeForDelete(AssetPtr);
+		}
+	}
+}
+
+void FPythonScriptPlugin::OnAssetReload(const EPackageReloadPhase InPackageReloadPhase, FPackageReloadedEvent* InPackageReloadedEvent)
+{
+	if (InPackageReloadPhase == EPackageReloadPhase::PostPackageFixup)
+	{
+		// Get the primary asset in this package
+		// Use the new package as it has the correct name
+		const UPackage* NewPackage = InPackageReloadedEvent->GetNewPackage();
+		const UObject* NewAsset = StaticFindObject(UObject::StaticClass(), (UPackage*)NewPackage, *FPackageName::GetLongPackageAssetName(NewPackage->GetName()));
+		OnAssetUpdated(NewAsset);
+	}
+}
+
+void FPythonScriptPlugin::OnAssetUpdated(const UObject* InObj)
+{
+	const UObject* AssetPtr = PyGenUtil::GetTypeRegistryType(InObj);
+	if (AssetPtr)
+	{
+		// If this asset has an associated Python type, then we need to re-generate it
+		FPyWrapperTypeRegistry& PyWrapperTypeRegistry = FPyWrapperTypeRegistry::Get();
+		if (PyWrapperTypeRegistry.HasWrappedTypeForObject(AssetPtr))
+		{
+			FPyWrapperTypeRegistry::FGeneratedWrappedTypeReferences GeneratedWrappedTypeReferences;
+			TSet<FName> DirtyModules;
+
+			PyWrapperTypeRegistry.GenerateWrappedTypeForObject(AssetPtr, GeneratedWrappedTypeReferences, DirtyModules, EPyTypeGenerationFlags::IncludeBlueprintGeneratedTypes | EPyTypeGenerationFlags::OverwriteExisting);
+
+			PyWrapperTypeRegistry.GenerateWrappedTypesForReferences(GeneratedWrappedTypeReferences, DirtyModules);
+			PyWrapperTypeRegistry.NotifyModulesDirtied(DirtyModules);
+		}
+	}
+}
+
 #if WITH_EDITOR
 
 void FPythonScriptPlugin::OnPrepareToCleanseEditorObject(UObject* InObject)
diff --git a/Engine/Plugins/Experimental/PythonScriptPlugin/Source/PythonScriptPlugin/Private/PythonScriptPlugin.h b/Engine/Plugins/Experimental/PythonScriptPlugin/Source/PythonScriptPlugin/Private/PythonScriptPlugin.h
--- a/Engine/Plugins/Experimental/PythonScriptPlugin/Source/PythonScriptPlugin/Private/PythonScriptPlugin.h
+++ b/Engine/Plugins/Experimental/PythonScriptPlugin/Source/PythonScriptPlugin/Private/PythonScriptPlugin.h
@@ -9,9 +9,15 @@
 #include "Misc/CoreMisc.h"
 #include "HAL/IConsoleManager.h"
 #include "Framework/Commands/InputChord.h"
+#include "Kismet2/EnumEditorUtils.h"
 
 class FPythonScriptPlugin;
 class FPythonScriptRemoteExecution;
+class FPackageReloadedEvent;
+
+struct FAssetData;
+
+enum class EPackageReloadPhase : uint8;
 
 #if WITH_PYTHON
 
@@ -81,7 +87,10 @@ struct IPythonCommandMenu
 };
 #endif	// WITH_PYTHON
 
-class FPythonScriptPlugin : public IPythonScriptPlugin, public FSelfRegisteringExec
+class FPythonScriptPlugin 
+	: public IPythonScriptPlugin
+	, public FSelfRegisteringExec
+	, public FEnumEditorUtils::INotifyOnEnumChanged
 {
 public:
 	FPythonScriptPlugin();
@@ -106,6 +115,10 @@ public:
 	//~ FSelfRegisteringExec interface
 	virtual bool Exec(UWorld* InWorld, const TCHAR* Cmd, FOutputDevice& Ar) override;
 
+	//~ FEnumEditorUtils::INotifyOnEnumChanged interface
+	virtual void PreChange(const UUserDefinedEnum* Enum, FEnumEditorUtils::EEnumEditorChangeInfo Info) override;
+	virtual void PostChange(const UUserDefinedEnum* Enum, FEnumEditorUtils::EEnumEditorChangeInfo Info) override;
+
 #if WITH_PYTHON
 
 #if WITH_EDITOR	
@@ -152,6 +165,14 @@ private:
 
 	void OnContentPathDismounted(const FString& InAssetPath, const FString& InFilesystemPath);
 
+	void OnAssetRenamed(const FAssetData& Data, const FString& OldName);
+
+	void OnAssetRemoved(const FAssetData& Data);
+
+	void OnAssetReload(const EPackageReloadPhase InPackageReloadPhase, FPackageReloadedEvent* InPackageReloadedEvent);
+
+	void OnAssetUpdated(const UObject* InObj);
+
 #if WITH_EDITOR
 	void OnPrepareToCleanseEditorObject(UObject* InObject);
 #endif	// WITH_EDITOR
diff --git a/Engine/Plugins/Experimental/PythonScriptPlugin/Source/PythonScriptPlugin/PythonScriptPlugin.Build.cs b/Engine/Plugins/Experimental/PythonScriptPlugin/Source/PythonScriptPlugin/PythonScriptPlugin.Build.cs
--- a/Engine/Plugins/Experimental/PythonScriptPlugin/Source/PythonScriptPlugin/PythonScriptPlugin.Build.cs
+++ b/Engine/Plugins/Experimental/PythonScriptPlugin/Source/PythonScriptPlugin/PythonScriptPlugin.Build.cs
@@ -18,6 +18,7 @@ namespace UnrealBuildTool.Rules
 			PrivateDependencyModuleNames.AddRange(
 				new string[] {
 					"Analytics",
+					"AssetRegistry",
 					"Projects",
 					"Python",
 					"Slate",
