Added set_editor_properties to let you set multiple properties which only emit a single pre/post change notification

hism_comp.set_editor_properties({
    'cast_shadow': False,
    'instance_start_cull_distance': 12,
    'instance_end_cull_distance': 23,
    })

diff --git a/Engine/Plugins/Experimental/PythonScriptPlugin/Source/PythonScriptPlugin/Private/PyPtr.h b/Engine/Plugins/Experimental/PythonScriptPlugin/Source/PythonScriptPlugin/Private/PyPtr.h
--- a/Engine/Plugins/Experimental/PythonScriptPlugin/Source/PythonScriptPlugin/Private/PyPtr.h
+++ b/Engine/Plugins/Experimental/PythonScriptPlugin/Source/PythonScriptPlugin/Private/PyPtr.h
@@ -126,12 +126,7 @@ public:
 		return Ptr != nullptr;
 	}
 
-	operator TPythonType*()
-	{
-		return Ptr;
-	}
-
-	operator const TPythonType*() const
+	operator TPythonType*() const
 	{
 		return Ptr;
 	}
@@ -165,17 +160,7 @@ public:
 		return Ptr;
 	}
 
-	const TPythonType*& Get() const
-	{
-		return Ptr;
-	}
-
-	TPythonType* GetPtr()
-	{
-		return Ptr;
-	}
-
-	const TPythonType* GetPtr() const
+	TPythonType* GetPtr() const
 	{
 		return Ptr;
 	}
diff --git a/Engine/Plugins/Experimental/PythonScriptPlugin/Source/PythonScriptPlugin/Private/PyWrapperObject.cpp b/Engine/Plugins/Experimental/PythonScriptPlugin/Source/PythonScriptPlugin/Private/PyWrapperObject.cpp
--- a/Engine/Plugins/Experimental/PythonScriptPlugin/Source/PythonScriptPlugin/Private/PyWrapperObject.cpp
+++ b/Engine/Plugins/Experimental/PythonScriptPlugin/Source/PythonScriptPlugin/Private/PyWrapperObject.cpp
@@ -16,6 +16,7 @@
 #include "UObject/StructOnScope.h"
 #include "UObject/PropertyPortFlags.h"
 #include "Engine/World.h"
+#include "Misc/ScopeExit.h"
 #include "Templates/Casts.h"
 #include "BlueprintActionDatabase.h"
 
@@ -730,7 +731,7 @@ PyTypeObject InitializePyWrapperObjectType()
 			FName NewName;
 			if (PyNameObj && PyNameObj != Py_None && !PyConversion::Nativize(PyNameObj, NewName))
 			{
-				PyUtil::SetPythonError(PyExc_TypeError, InSelf, *FString::Printf(TEXT("Failed to convert 'name' (%s) to 'Name'"), *PyUtil::GetFriendlyTypename(InSelf)));
+				PyUtil::SetPythonError(PyExc_TypeError, InSelf, *FString::Printf(TEXT("Failed to convert 'name' (%s) to 'Name'"), *PyUtil::GetFriendlyTypename(PyNameObj)));
 				return nullptr;
 			}
 
@@ -746,48 +747,29 @@ PyTypeObject InitializePyWrapperObjectType()
 			return PyConversion::Pythonize(bResult);
 		}
 
-		static PyObject* GetEditorProperty(FPyWrapperObject* InSelf, PyObject* InArgs, PyObject* InKwds)
+		static PyGenUtil::FGeneratedWrappedProperty FindEditorPropertyImpl(FPyWrapperObject* InSelf, const FName InPythonPropName)
 		{
-			if (!FPyWrapperObject::ValidateInternalState(InSelf))
-			{
-				return nullptr;
-			}
-
-			PyObject* PyNameObj = nullptr;
-
-			static const char *ArgsKwdList[] = { "name", nullptr };
-			if (!PyArg_ParseTupleAndKeywords(InArgs, InKwds, "O:get_editor_property", (char**)ArgsKwdList, &PyNameObj))
-			{
-				return nullptr;
-			}
-
-			FName Name;
-			if (!PyConversion::Nativize(PyNameObj, Name))
-			{
-				PyUtil::SetPythonError(PyExc_TypeError, InSelf, *FString::Printf(TEXT("Failed to convert 'name' (%s) to 'Name'"), *PyUtil::GetFriendlyTypename(InSelf)));
-				return nullptr;
-			}
+			PyGenUtil::FGeneratedWrappedProperty WrappedPropDef;
 
 			const UClass* Class = InSelf->ObjectInstance->GetClass();
 
-			const FName ResolvedName = FPyWrapperObjectMetaData::ResolvePropertyName(InSelf, Name);
+			const FName ResolvedName = FPyWrapperObjectMetaData::ResolvePropertyName(InSelf, InPythonPropName);
 			const UProperty* ResolvedProp = Class->FindPropertyByName(ResolvedName);
 			if (!ResolvedProp)
 			{
-				PyUtil::SetPythonError(PyExc_Exception, InSelf, *FString::Printf(TEXT("Failed to find property '%s' for attribute '%s' on '%s'"), *ResolvedName.ToString(), *Name.ToString(), *Class->GetName()));
-				return nullptr;
+				PyUtil::SetPythonError(PyExc_Exception, InSelf, *FString::Printf(TEXT("Failed to find property '%s' for attribute '%s' on '%s'"), *ResolvedName.ToString(), *InPythonPropName.ToString(), *Class->GetName()));
+				return WrappedPropDef;
 			}
 
 			TOptional<FString> PropDeprecationMessage;
 			{
 				FString PropDeprecationMessageStr;
-				if (FPyWrapperObjectMetaData::IsPropertyDeprecated(InSelf, Name, &PropDeprecationMessageStr))
+				if (FPyWrapperObjectMetaData::IsPropertyDeprecated(InSelf, InPythonPropName, &PropDeprecationMessageStr))
 				{
 					PropDeprecationMessage = MoveTemp(PropDeprecationMessageStr);
 				}
 			}
-
-			PyGenUtil::FGeneratedWrappedProperty WrappedPropDef;
+			
 			if (PropDeprecationMessage.IsSet())
 			{
 				WrappedPropDef.SetProperty(ResolvedProp, PyGenUtil::FGeneratedWrappedProperty::SPF_None);
@@ -798,6 +780,37 @@ PyTypeObject InitializePyWrapperObjectType()
 				WrappedPropDef.SetProperty(ResolvedProp);
 			}
 
+			return WrappedPropDef;
+		}
+
+		static PyObject* GetEditorProperty(FPyWrapperObject* InSelf, PyObject* InArgs, PyObject* InKwds)
+		{
+			if (!FPyWrapperObject::ValidateInternalState(InSelf))
+			{
+				return nullptr;
+			}
+
+			PyObject* PyNameObj = nullptr;
+
+			static const char *ArgsKwdList[] = { "name", nullptr };
+			if (!PyArg_ParseTupleAndKeywords(InArgs, InKwds, "O:get_editor_property", (char**)ArgsKwdList, &PyNameObj))
+			{
+				return nullptr;
+			}
+
+			FName Name;
+			if (!PyConversion::Nativize(PyNameObj, Name))
+			{
+				PyUtil::SetPythonError(PyExc_TypeError, InSelf, *FString::Printf(TEXT("Failed to convert 'name' (%s) to 'Name'"), *PyUtil::GetFriendlyTypename(PyNameObj)));
+				return nullptr;
+			}
+
+			const PyGenUtil::FGeneratedWrappedProperty WrappedPropDef = FindEditorPropertyImpl(InSelf, Name);
+			if (!WrappedPropDef.Prop)
+			{
+				return nullptr;
+			}
+
 			return FPyWrapperObject::GetPropertyValue(InSelf, WrappedPropDef, TCHAR_TO_UTF8(*Name.ToString()));
 		}
 
@@ -820,46 +833,104 @@ PyTypeObject InitializePyWrapperObjectType()
 			FName Name;
 			if (!PyConversion::Nativize(PyNameObj, Name))
 			{
-				PyUtil::SetPythonError(PyExc_TypeError, InSelf, *FString::Printf(TEXT("Failed to convert 'name' (%s) to 'Name'"), *PyUtil::GetFriendlyTypename(InSelf)));
+				PyUtil::SetPythonError(PyExc_TypeError, InSelf, *FString::Printf(TEXT("Failed to convert 'name' (%s) to 'Name'"), *PyUtil::GetFriendlyTypename(PyNameObj)));
 				return nullptr;
 			}
 
-			const UClass* Class = InSelf->ObjectInstance->GetClass();
-
-			const FName ResolvedName = FPyWrapperObjectMetaData::ResolvePropertyName(InSelf, Name);
-			const UProperty* ResolvedProp = Class->FindPropertyByName(ResolvedName);
-			if (!ResolvedProp)
+			const PyGenUtil::FGeneratedWrappedProperty WrappedPropDef = FindEditorPropertyImpl(InSelf, Name);
+			if (!WrappedPropDef.Prop)
 			{
-				PyUtil::SetPythonError(PyExc_Exception, InSelf, *FString::Printf(TEXT("Failed to find property '%s' for attribute '%s' on '%s'"), *ResolvedName.ToString(), *Name.ToString(), *Class->GetName()));
 				return nullptr;
 			}
 
-			TOptional<FString> PropDeprecationMessage;
+			const int Result = FPyWrapperObject::SetPropertyValue(InSelf, PyValueObj, WrappedPropDef, TCHAR_TO_UTF8(*Name.ToString()), /*InNotifyChange*/true, PropertyAccessUtil::EditorReadOnlyFlags);
+			if (Result != 0)
 			{
-				FString PropDeprecationMessageStr;
-				if (FPyWrapperObjectMetaData::IsPropertyDeprecated(InSelf, Name, &PropDeprecationMessageStr))
-				{
-					PropDeprecationMessage = MoveTemp(PropDeprecationMessageStr);
-				}
+				return nullptr;
 			}
 
-			PyGenUtil::FGeneratedWrappedProperty WrappedPropDef;
-			if (PropDeprecationMessage.IsSet())
+			Py_RETURN_NONE;
+		}
+
+		static PyObject* SetEditorProperties(FPyWrapperObject* InSelf, PyObject* InArgs, PyObject* InKwds)
+		{
+			if (!FPyWrapperObject::ValidateInternalState(InSelf))
 			{
-				WrappedPropDef.SetProperty(ResolvedProp, PyGenUtil::FGeneratedWrappedProperty::SPF_None);
-				WrappedPropDef.DeprecationMessage = MoveTemp(PropDeprecationMessage);
+				return nullptr;
 			}
-			else
+
+			PyObject* PyPropertyInfoDict = nullptr;
+			if (!PyArg_ParseTuple(InArgs, "O:set_editor_properties", &PyPropertyInfoDict))
 			{
-				WrappedPropDef.SetProperty(ResolvedProp);
+				return nullptr;
 			}
 
-			const int Result = FPyWrapperObject::SetPropertyValue(InSelf, PyValueObj, WrappedPropDef, TCHAR_TO_UTF8(*Name.ToString()), /*InNotifyChange*/true, PropertyAccessUtil::EditorReadOnlyFlags);
-			if (Result != 0)
+			if (!PyUtil::IsMappingType(PyPropertyInfoDict))
 			{
+				PyUtil::SetPythonError(PyExc_TypeError, InSelf, *FString::Printf(TEXT("'property_info' (%s) is expected to be a mapping type (eg, dict)"), *PyUtil::GetFriendlyTypename(PyPropertyInfoDict)));
 				return nullptr;
 			}
 
+			// Build up the list of properties and their new values
+			typedef TTuple<FName, PyGenUtil::FGeneratedWrappedProperty, FPyObjectPtr> FPropertyInfoPair;
+			TArray<FPropertyInfoPair> PropertyInfos;
+			{
+				FPyObjectPtr PyPropertyInfoDictIter = FPyObjectPtr::StealReference(PyObject_GetIter(PyPropertyInfoDict));
+				if (PyPropertyInfoDictIter)
+				{
+					FPyObjectPtr PyKeyItem;
+					while (PyKeyItem = FPyObjectPtr::StealReference(PyIter_Next(PyPropertyInfoDictIter)))
+					{
+						FName Name;
+						if (!PyConversion::Nativize(PyKeyItem, Name))
+						{
+							PyUtil::SetPythonError(PyExc_TypeError, InSelf, *FString::Printf(TEXT("Failed to convert dict key (%s) to 'Name'"), *PyUtil::GetFriendlyTypename(PyKeyItem)));
+							return nullptr;
+						}
+
+						PyGenUtil::FGeneratedWrappedProperty WrappedPropDef = FindEditorPropertyImpl(InSelf, Name);
+						if (!WrappedPropDef.Prop)
+						{
+							return nullptr;
+						}
+
+						FPyObjectPtr PyValueItem = FPyObjectPtr::StealReference(PyObject_GetItem(PyPropertyInfoDict, PyKeyItem));
+						if (ensure(PyValueItem))
+						{
+							PropertyInfos.Add(MakeTuple(Name, MoveTemp(WrappedPropDef), MoveTemp(PyValueItem)));
+						}
+					}
+
+					// Iteration error?
+					if (PyErr_Occurred())
+					{
+						return nullptr;
+					}
+				}
+			}
+
+			// At this point we know that every property we were asked to set was found, but not that the values are going to be compatible
+			// Checking value coercion is tricky without actually doing the conversion, so we'll assume that this point that we need to do the change notifies
+			InSelf->ObjectInstance->PreEditChange(nullptr);
+			ON_SCOPE_EXIT
+			{
+				InSelf->ObjectInstance->PostEditChange();
+			};
+
+			// Try and set the value of each property
+			for (const FPropertyInfoPair& PropertyInfo : PropertyInfos)
+			{
+				const FName Name = PropertyInfo.Get<0>();
+				const PyGenUtil::FGeneratedWrappedProperty& WrappedPropDef = PropertyInfo.Get<1>();
+				PyObject* PyValueObj = PropertyInfo.Get<2>().GetPtr();
+
+				const int Result = FPyWrapperObject::SetPropertyValue(InSelf, PyValueObj, WrappedPropDef, TCHAR_TO_UTF8(*Name.ToString()), /*InNotifyChange*/false, PropertyAccessUtil::EditorReadOnlyFlags);
+				if (Result != 0)
+				{
+					return nullptr;
+				}
+			}
+
 			Py_RETURN_NONE;
 		}
 	};
@@ -882,6 +953,7 @@ PyTypeObject InitializePyWrapperObjectType()
 		{ "rename", PyCFunctionCast(&FMethods::Rename), METH_VARARGS | METH_KEYWORDS, "x.rename(name=None, outer=None) -> bool -- rename this instance" },
 		{ "get_editor_property", PyCFunctionCast(&FMethods::GetEditorProperty), METH_VARARGS | METH_KEYWORDS, "x.get_editor_property(name) -> object -- get the value of any property visible to the editor" },
 		{ "set_editor_property", PyCFunctionCast(&FMethods::SetEditorProperty), METH_VARARGS | METH_KEYWORDS, "x.set_editor_property(name, value) -> None -- set the value of any property visible to the editor, ensuring that the pre/post change notifications are called" },
+		{ "set_editor_properties", PyCFunctionCast(&FMethods::SetEditorProperties), METH_VARARGS, "x.set_editor_properties(property_info) -> None -- set the value of any properties visible to the editor (from a name->value dict), ensuring that the pre/post change notifications are called" },
 		{ nullptr, nullptr, 0, nullptr }
 	};
 
