CIS fix

#rnx

diff --git a/Engine/Plugins/Experimental/PythonScriptPlugin/Source/PythonScriptPlugin/Private/PyWrapperObject.cpp b/Engine/Plugins/Experimental/PythonScriptPlugin/Source/PythonScriptPlugin/Private/PyWrapperObject.cpp
--- a/Engine/Plugins/Experimental/PythonScriptPlugin/Source/PythonScriptPlugin/Private/PyWrapperObject.cpp
+++ b/Engine/Plugins/Experimental/PythonScriptPlugin/Source/PythonScriptPlugin/Private/PyWrapperObject.cpp
@@ -879,7 +879,7 @@ PyTypeObject InitializePyWrapperObjectType()
 				if (PyPropertyInfoDictIter)
 				{
 					FPyObjectPtr PyKeyItem;
-					while (PyKeyItem = FPyObjectPtr::StealReference(PyIter_Next(PyPropertyInfoDictIter)))
+					while ((PyKeyItem = FPyObjectPtr::StealReference(PyIter_Next(PyPropertyInfoDictIter))))
 					{
 						FName Name;
 						if (!PyConversion::Nativize(PyKeyItem, Name))
