CIS fix

#rnx

diff --git a/Engine/Plugins/Experimental/PythonScriptPlugin/Source/PythonScriptPlugin/Private/PyEngine.cpp b/Engine/Plugins/Experimental/PythonScriptPlugin/Source/PythonScriptPlugin/Private/PyEngine.cpp
--- a/Engine/Plugins/Experimental/PythonScriptPlugin/Source/PythonScriptPlugin/Private/PyEngine.cpp
+++ b/Engine/Plugins/Experimental/PythonScriptPlugin/Source/PythonScriptPlugin/Private/PyEngine.cpp
@@ -126,7 +126,7 @@ PyObject* GetBlueprintGeneratedTypes(PyObject* InSelf, PyObject* InArgs)
 							FPyObjectPtr ValueItem = FPyObjectPtr::StealReference(PyIter_Next(PyObjIter));
 							if (!ValueItem)
 							{
-								return false;
+								return nullptr;
 							}
 
 							FString& AssetToGenerate = AssetsToGenerate.AddDefaulted_GetRef();
