Fixed Python memory leak

#rnx

diff --git a/Engine/Plugins/Experimental/PythonScriptPlugin/Source/PythonScriptPlugin/Private/PyGenUtil.cpp b/Engine/Plugins/Experimental/PythonScriptPlugin/Source/PythonScriptPlugin/Private/PyGenUtil.cpp
--- a/Engine/Plugins/Experimental/PythonScriptPlugin/Source/PythonScriptPlugin/Private/PyGenUtil.cpp
+++ b/Engine/Plugins/Experimental/PythonScriptPlugin/Source/PythonScriptPlugin/Private/PyGenUtil.cpp
@@ -605,7 +605,11 @@ bool FGeneratedWrappedType::Finalize()
 		else if (FinalizedState == EFinalizedState::Reset)
 		{
 			PyType_Modified(&PyType);
-			PyDict_SetItemString(PyType.tp_dict, "__doc__", PyUnicode_FromString(PyType.tp_doc ? PyType.tp_doc : "")); // PyType_Modified doesn't update __doc__
+
+			// PyType_Modified doesn't update __doc__
+			FPyObjectPtr PyDocString = FPyObjectPtr::StealReference(PyUnicode_FromString(PyType.tp_doc ? PyType.tp_doc : ""));
+			PyDict_SetItemString(PyType.tp_dict, "__doc__", PyDocString);
+
 			bSuccess = true;
 		}
 		FinalizedState = EFinalizedState::Finalized;
