Python FProperty fixes

 - Fixed Python container wrappers leaking transient property instances.
   - PyUtil::TPropOnScope will manage an optionally owned property instance, deleting it if required when it goes out-of-scope.
 - Fixed Python container wrappers crashing during GC due to calling FReferenceCollector::AddReferencedObjects passing a FProperty-based type.
 - Fixed Python delegate wrappers crashing when called with arguments due to checking Children rather than ChildProperties.

#rnx

diff --git a/Engine/Plugins/Experimental/PythonScriptPlugin/Source/PythonScriptPlugin/Private/PyCore.cpp b/Engine/Plugins/Experimental/PythonScriptPlugin/Source/PythonScriptPlugin/Private/PyCore.cpp
--- a/Engine/Plugins/Experimental/PythonScriptPlugin/Source/PythonScriptPlugin/Private/PyCore.cpp
+++ b/Engine/Plugins/Experimental/PythonScriptPlugin/Source/PythonScriptPlugin/Private/PyCore.cpp
@@ -80,23 +80,12 @@ void ApplyMetaData(PyObject* InMetaData, const TFunctionRef<void(const FString&,
 
 } // namespace PyCoreUtil
 
-TStrongObjectPtr<UStruct>& GetPythonPropertyContainerSingleton()
-{
-	static TStrongObjectPtr<UStruct> PythonPropertyContainer;
-	return PythonPropertyContainer;
-}
-
 TStrongObjectPtr<UPackage>& GetPythonTypeContainerSingleton()
 {
 	static TStrongObjectPtr<UPackage> PythonTypeContainer;
 	return PythonTypeContainer;
 }
 
-UObject* GetPythonPropertyContainer()
-{
-	return GetPythonPropertyContainerSingleton().Get();
-}
-
 UObject* GetPythonTypeContainer()
 {
 	return GetPythonTypeContainerSingleton().Get();
@@ -1692,8 +1681,6 @@ PyMethodDef PyCoreMethods[] = {
 
 void InitializeModule()
 {
-	GetPythonPropertyContainerSingleton().Reset(NewObject<UStruct>(GetTransientPackage(), TEXT("PythonProperties"), RF_Transient));
-
 	GetPythonTypeContainerSingleton().Reset(NewObject<UPackage>(nullptr, TEXT("/Engine/PythonTypes"), RF_Public | RF_Transient));
 	GetPythonTypeContainerSingleton()->SetPackageFlags(PKG_ContainsScript);
 
diff --git a/Engine/Plugins/Experimental/PythonScriptPlugin/Source/PythonScriptPlugin/Private/PyCore.h b/Engine/Plugins/Experimental/PythonScriptPlugin/Source/PythonScriptPlugin/Private/PyCore.h
--- a/Engine/Plugins/Experimental/PythonScriptPlugin/Source/PythonScriptPlugin/Private/PyCore.h
+++ b/Engine/Plugins/Experimental/PythonScriptPlugin/Source/PythonScriptPlugin/Private/PyCore.h
@@ -18,9 +18,6 @@ enum class EPythonLogOutputType : uint8;
 
 struct FSlowTask;
 
-/** Get the object that Python created transient properties should be outered to */
-UObject* GetPythonPropertyContainer();
-
 /** Get the object that Python created types should be outered to */
 UObject* GetPythonTypeContainer();
 
diff --git a/Engine/Plugins/Experimental/PythonScriptPlugin/Source/PythonScriptPlugin/Private/PyGenUtil.cpp b/Engine/Plugins/Experimental/PythonScriptPlugin/Source/PythonScriptPlugin/Private/PyGenUtil.cpp
--- a/Engine/Plugins/Experimental/PythonScriptPlugin/Source/PythonScriptPlugin/Private/PyGenUtil.cpp
+++ b/Engine/Plugins/Experimental/PythonScriptPlugin/Source/PythonScriptPlugin/Private/PyGenUtil.cpp
@@ -2260,7 +2260,7 @@ FString PythonizeValue(const FProperty* InProp, const void* InPropValue, const u
 
 FString PythonizeDefaultValue(const FProperty* InProp, const FString& InDefaultValue, const uint32 InFlags)
 {
-	PyUtil::FPropValueOnScope PropValue(InProp);
+	PyUtil::FPropValueOnScope PropValue(PyUtil::FConstPropOnScope::ExternalReference(InProp));
 	PyUtil::ImportDefaultValue(InProp, PropValue.GetValue(), InDefaultValue);
 	return PythonizeValue(InProp, PropValue.GetValue(), InFlags);
 }
diff --git a/Engine/Plugins/Experimental/PythonScriptPlugin/Source/PythonScriptPlugin/Private/PyUtil.cpp b/Engine/Plugins/Experimental/PythonScriptPlugin/Source/PythonScriptPlugin/Private/PyUtil.cpp
--- a/Engine/Plugins/Experimental/PythonScriptPlugin/Source/PythonScriptPlugin/Private/PyUtil.cpp
+++ b/Engine/Plugins/Experimental/PythonScriptPlugin/Source/PythonScriptPlugin/Private/PyUtil.cpp
@@ -36,6 +36,8 @@ DEFINE_LOG_CATEGORY(LogPython);
 namespace PyUtil
 {
 
+const FName DefaultPythonPropertyName = "TransientPythonProperty";
+
 FPyApiBuffer TCHARToPyApiBuffer(const TCHAR* InStr)
 {
 	auto PyCharToPyBuffer = [](const FPyApiChar* InPyChar)
@@ -100,9 +102,8 @@ FString PyObjectToUEStringRepr(PyObject* InPyObj)
 	return PyObjectToUEString(InPyObj);
 }
 
-FPropValueOnScope::FPropValueOnScope(const FProperty* InProp)
-	: Prop(InProp)
-	, Value(nullptr)
+FPropValueOnScope::FPropValueOnScope(FConstPropOnScope&& InProp)
+	: Prop(MoveTemp(InProp))
 {
 	check(Prop);
 
@@ -149,27 +150,27 @@ void* FPropValueOnScope::GetValue(const int32 InArrayIndex) const
 }
 
 FFixedArrayElementOnScope::FFixedArrayElementOnScope(const FProperty* InProp)
-	: FPropValueOnScope(PyUtil::CreateProperty(InProp)) // We have to create a new temporary property with an ArrayDim of 1
+	: FPropValueOnScope(FConstPropOnScope::OwnedReference(PyUtil::CreateProperty(InProp))) // We have to create a new temporary property with an ArrayDim of 1
 {
 }
 
 FArrayElementOnScope::FArrayElementOnScope(const FArrayProperty* InProp)
-	: FPropValueOnScope(InProp->Inner)
+	: FPropValueOnScope(FConstPropOnScope::ExternalReference(InProp->Inner))
 {
 }
 
 FSetElementOnScope::FSetElementOnScope(const FSetProperty* InProp)
-	: FPropValueOnScope(InProp->ElementProp)
+	: FPropValueOnScope(FConstPropOnScope::ExternalReference(InProp->ElementProp))
 {
 }
 
 FMapKeyOnScope::FMapKeyOnScope(const FMapProperty* InProp)
-	: FPropValueOnScope(InProp->KeyProp)
+	: FPropValueOnScope(FConstPropOnScope::ExternalReference(InProp->KeyProp))
 {
 }
 
 FMapValueOnScope::FMapValueOnScope(const FMapProperty* InProp)
-	: FPropValueOnScope(InProp->ValueProp)
+	: FPropValueOnScope(FConstPropOnScope::ExternalReference(InProp->ValueProp))
 {
 }
 
@@ -365,8 +366,7 @@ bool CalculatePropertyDef(PyObject* InPyObj, FPropertyDef& OutPropertyDef)
 FProperty* CreateProperty(const FPropertyDef& InPropertyDef, const int32 InArrayDim, UObject* InOuter, const FName InName)
 {
 	check(InArrayDim > 0);
-	UObject* PropOuter = InOuter ? InOuter : GetPythonPropertyContainer();
-	FProperty* Prop = CastFieldChecked<FProperty>(InPropertyDef.PropertyClass->Construct(PropOuter, InName, RF_NoFlags));
+	FProperty* Prop = CastFieldChecked<FProperty>(InPropertyDef.PropertyClass->Construct(InOuter, InName, RF_NoFlags));
 	if (Prop)
 	{
 		Prop->ArrayDim = InArrayDim;
diff --git a/Engine/Plugins/Experimental/PythonScriptPlugin/Source/PythonScriptPlugin/Private/PyUtil.h b/Engine/Plugins/Experimental/PythonScriptPlugin/Source/PythonScriptPlugin/Private/PyUtil.h
--- a/Engine/Plugins/Experimental/PythonScriptPlugin/Source/PythonScriptPlugin/Private/PyUtil.h
+++ b/Engine/Plugins/Experimental/PythonScriptPlugin/Source/PythonScriptPlugin/Private/PyUtil.h
@@ -49,6 +49,8 @@ namespace PyUtil
 	#define TCHARToPyApiChar(InStr) TCHAR_TO_UTF8(InStr)
 #endif	// PY_MAJOR_VERSION >= 3
 
+	extern const FName DefaultPythonPropertyName;
+
 	/** Convert a TCHAR to a persistent buffer that can be passed to a Python API that does hold the result (you have to keep the buffer around as long as Python needs it) */
 	FPyApiBuffer TCHARToPyApiBuffer(const TCHAR* InStr);
 
@@ -75,16 +77,153 @@ namespace PyUtil
 		return PyBool_FromLong(bResult);
 	}
 
+	/** Helper to manage a property pointer that may be potentially owned by a wrapper or stack instance */
+	template <typename TPropType>
+	class TPropOnScope
+	{
+		static_assert(TIsDerivedFrom<TPropType, FProperty>::IsDerived, "TPropType must be a FProperty-based type!");
+
+	public:
+		static TPropOnScope OwnedReference(TPropType* InProp)
+		{
+			return TPropOnScope(InProp, true);
+		}
+
+		static TPropOnScope ExternalReference(TPropType* InProp)
+		{
+			return TPropOnScope(InProp, false);
+		}
+
+		TPropOnScope() = default;
+
+		~TPropOnScope()
+		{
+			Reset();
+		}
+
+		TPropOnScope(const TPropOnScope&) = delete;
+		TPropOnScope& operator=(const TPropOnScope&) = delete;
+
+		template <
+			typename TOtherPropType,
+			typename = decltype(ImplicitConv<TPropType*>((TOtherPropType*)nullptr))
+		>
+		TPropOnScope(TPropOnScope<TOtherPropType>&& Other)
+		{
+			bOwnsProp = Other.OwnsProp();
+			Prop = Other.Release();
+
+			Other.Reset();
+		}
+
+		template <
+			typename TOtherPropType,
+			typename = decltype(ImplicitConv<TPropType*>((TOtherPropType*)nullptr))
+		>
+		TPropOnScope& operator=(TPropOnScope<TOtherPropType>&& Other)
+		{
+			if (this != (void*)&Other)
+			{
+				Reset();
+
+				bOwnsProp = Other.OwnsProp();
+				Prop = Other.Release();
+
+				Other.Reset();
+			}
+			return *this;
+		}
+
+		explicit operator bool() const
+		{
+			return IsValid();
+		}
+
+		bool IsValid() const
+		{
+			return Prop != nullptr;
+		}
+
+		operator TPropType*() const
+		{
+			return Prop;
+		}
+
+		TPropType& operator*() const
+		{
+			check(Prop);
+			return *Prop;
+		}
+
+		TPropType* operator->() const
+		{
+			check(Prop);
+			return Prop;
+		}
+
+		TPropType* Get() const
+		{
+			return Prop;
+		}
+
+		bool OwnsProp() const
+		{
+			return bOwnsProp;
+		}
+
+		TPropType* Release()
+		{
+			TPropType* LocalProp = Prop;
+			Prop = nullptr;
+			bOwnsProp = false;
+			return LocalProp;
+		}
+
+		void Reset()
+		{
+			if (bOwnsProp)
+			{
+				delete Prop;
+			}
+			Prop = nullptr;
+			bOwnsProp = false;
+		}
+
+		void AddReferencedObjects(FReferenceCollector& Collector)
+		{
+			if (Prop)
+			{
+				((typename TRemoveConst<TPropType>::Type*)Prop)->AddReferencedObjects(Collector);
+			}
+		}
+
+	private:
+		TPropOnScope(TPropType* InProp, const bool InOwnsProp)
+			: Prop(InProp)
+			, bOwnsProp(InOwnsProp)
+		{
+		}
+
+		TPropType* Prop = nullptr;
+		bool bOwnsProp = false;
+	};
+
+	typedef TPropOnScope<FProperty> FPropOnScope;
+	typedef TPropOnScope<const FProperty> FConstPropOnScope;
+	typedef TPropOnScope<FArrayProperty> FArrayPropOnScope;
+	typedef TPropOnScope<const FArrayProperty> FConstArrayPropOnScope;
+	typedef TPropOnScope<FSetProperty> FSetPropOnScope;
+	typedef TPropOnScope<const FSetProperty> FConstSetPropOnScope;
+	typedef TPropOnScope<FMapProperty> FMapPropOnScope;
+	typedef TPropOnScope<const FMapProperty> FConstMapPropOnScope;
+
 	/** Helper used to hold the value for a property value on the stack */
 	class FPropValueOnScope
 	{
 	public:
-		explicit FPropValueOnScope(const FProperty* InProp);
+		explicit FPropValueOnScope(FConstPropOnScope&& InProp);
 		~FPropValueOnScope();
 
-		FPropValueOnScope(FPropValueOnScope&) = delete;
-		FPropValueOnScope& operator=(FPropValueOnScope&) = delete;
-
 		bool SetValue(PyObject* InPyObj, const TCHAR* InErrorCtxt);
 
 		bool IsValid() const;
@@ -94,8 +233,8 @@ namespace PyUtil
 		void* GetValue(const int32 InArrayIndex = 0) const;
 
 	private:
-		const FProperty* Prop;
-		void* Value;
+		FConstPropOnScope Prop;
+		void* Value = nullptr;
 	};
 
 	/** Helper used to hold the value for a single fixed array element on the stack */
@@ -187,13 +326,13 @@ namespace PyUtil
 	bool CalculatePropertyDef(PyObject* InPyObj, FPropertyDef& OutPropertyDef);
 
 	/** Given a property definition, create a property instance */
-	FProperty* CreateProperty(const FPropertyDef& InPropertyDef, const int32 InArrayDim = 1, UObject* InOuter = nullptr, const FName InName = NAME_None);
+	FProperty* CreateProperty(const FPropertyDef& InPropertyDef, const int32 InArrayDim = 1, UObject* InOuter = nullptr, const FName InName = DefaultPythonPropertyName);
 
 	/** Given a Python type, create a compatible property instance */
-	FProperty* CreateProperty(PyTypeObject* InPyType, const int32 InArrayDim = 1, UObject* InOuter = nullptr, const FName InName = NAME_None);
+	FProperty* CreateProperty(PyTypeObject* InPyType, const int32 InArrayDim = 1, UObject* InOuter = nullptr, const FName InName = DefaultPythonPropertyName);
 
 	/** Given a Python instance, create a compatible property instance */
-	FProperty* CreateProperty(PyObject* InPyObj, const int32 InArrayDim = 1, UObject* InOuter = nullptr, const FName InName = NAME_None);
+	FProperty* CreateProperty(PyObject* InPyObj, const int32 InArrayDim = 1, UObject* InOuter = nullptr, const FName InName = DefaultPythonPropertyName);
 
 	/** Check to see if the given property is an input parameter for a function */
 	bool IsInputParameter(const FProperty* InParam);
diff --git a/Engine/Plugins/Experimental/PythonScriptPlugin/Source/PythonScriptPlugin/Private/PyWrapperArray.cpp b/Engine/Plugins/Experimental/PythonScriptPlugin/Source/PythonScriptPlugin/Private/PyWrapperArray.cpp
--- a/Engine/Plugins/Experimental/PythonScriptPlugin/Source/PythonScriptPlugin/Private/PyWrapperArray.cpp
+++ b/Engine/Plugins/Experimental/PythonScriptPlugin/Source/PythonScriptPlugin/Private/PyWrapperArray.cpp
@@ -2,8 +2,6 @@
 
 #include "PyWrapperArray.h"
 #include "PyWrapperTypeRegistry.h"
-#include "PyCore.h"
-#include "PyUtil.h"
 #include "PyConversion.h"
 #include "PyReferenceCollector.h"
 #include "UObject/UnrealType.h"
@@ -121,7 +119,7 @@ FPyWrapperArray* FPyWrapperArray::New(PyTypeObject* InType)
 	if (Self)
 	{
 		new(&Self->OwnerContext) FPyWrapperOwnerContext();
-		Self->ArrayProp = nullptr;
+		new(&Self->ArrayProp) PyUtil::FConstArrayPropOnScope();
 		Self->ArrayInstance = nullptr;
 	}
 	return Self;
@@ -132,6 +130,7 @@ void FPyWrapperArray::Free(FPyWrapperArray* InSelf)
 	Deinit(InSelf);
 
 	InSelf->OwnerContext.~FPyWrapperOwnerContext();
+	InSelf->ArrayProp.~TPropOnScope();
 	FPyWrapperBase::Free(InSelf);
 }
 
@@ -145,15 +144,15 @@ int FPyWrapperArray::Init(FPyWrapperArray* InSelf, const PyUtil::FPropertyDef& I
 		return BaseInit;
 	}
 
-	FProperty* ArrayElementProp = PyUtil::CreateProperty(InElementDef, 1);
+	PyUtil::FPropOnScope ArrayElementProp = PyUtil::FPropOnScope::OwnedReference(PyUtil::CreateProperty(InElementDef, 1));
 	if (!ArrayElementProp)
 	{
 		PyUtil::SetPythonError(PyExc_Exception, InSelf, TEXT("Array element property was null during init"));
 		return -1;
 	}
 
-	FArrayProperty* ArrayProp = new FArrayProperty(GetPythonPropertyContainer(), FField::GenerateFFieldName(GetPythonPropertyContainer(), FArrayProperty::StaticClass()), RF_NoFlags);
-	ArrayProp->Inner = ArrayElementProp;
+	PyUtil::FArrayPropOnScope ArrayProp = PyUtil::FArrayPropOnScope::OwnedReference(new FArrayProperty(FFieldVariant(), PyUtil::DefaultPythonPropertyName, RF_NoFlags));
+	ArrayProp->Inner = ArrayElementProp.Release();
 
 	// Need to manually call Link to fix-up some data (such as the C++ property flags) that are only set during Link
 	{
@@ -164,7 +163,7 @@ int FPyWrapperArray::Init(FPyWrapperArray* InSelf, const PyUtil::FPropertyDef& I
 	void* ArrayValue = FMemory::Malloc(ArrayProp->GetSize(), ArrayProp->GetMinAlignment());
 	ArrayProp->InitializeValue(ArrayValue);
 
-	InSelf->ArrayProp = ArrayProp;
+	InSelf->ArrayProp = MoveTemp(ArrayProp);
 	InSelf->ArrayInstance = ArrayValue;
 
 	FPyWrapperArrayFactory::Get().MapInstance(InSelf->ArrayInstance, InSelf);
@@ -185,23 +184,22 @@ int FPyWrapperArray::Init(FPyWrapperArray* InSelf, const FPyWrapperOwnerContext&
 
 	check(InProp && InValue);
 
-	const FArrayProperty* PropToUse = nullptr;
+	PyUtil::FConstArrayPropOnScope PropToUse;
 	void* ArrayInstanceToUse = nullptr;
 	switch (InConversionMethod)
 	{
 	case EPyConversionMethod::Copy:
 	case EPyConversionMethod::Steal:
 		{
-			FProperty* ArrayElementProp = PyUtil::CreateProperty(InProp->Inner);
+			PyUtil::FPropOnScope ArrayElementProp = PyUtil::FPropOnScope::OwnedReference(PyUtil::CreateProperty(InProp->Inner));
 			if (!ArrayElementProp)
 			{
 				PyUtil::SetPythonError(PyExc_TypeError, InSelf, *FString::Printf(TEXT("Failed to create element property from '%s' (%s)"), *InProp->Inner->GetName(), *InProp->Inner->GetClass()->GetName()));
 				return -1;
 			}
 
-			FArrayProperty* ArrayProp = new FArrayProperty(GetPythonPropertyContainer(), FField::GenerateFFieldName(GetPythonPropertyContainer(), FArrayProperty::StaticClass()), RF_NoFlags);
-			ArrayProp->Inner = ArrayElementProp;
-			PropToUse = ArrayProp;
+			PyUtil::FArrayPropOnScope ArrayProp = PyUtil::FArrayPropOnScope::OwnedReference(new FArrayProperty(FFieldVariant(), PyUtil::DefaultPythonPropertyName, RF_NoFlags));
+			ArrayProp->Inner = ArrayElementProp.Release();
 
 			// Need to manually call Link to fix-up some data (such as the C++ property flags) that are only set during Link
 			{
@@ -209,6 +207,8 @@ int FPyWrapperArray::Init(FPyWrapperArray* InSelf, const FPyWrapperOwnerContext&
 				ArrayProp->LinkWithoutChangingOffset(Ar);
 			}
 
+			PropToUse = MoveTemp(ArrayProp);
+
 			ArrayInstanceToUse = FMemory::Malloc(PropToUse->GetSize(), PropToUse->GetMinAlignment());
 			PropToUse->InitializeValue(ArrayInstanceToUse);
 			if (InConversionMethod == EPyConversionMethod::Steal)
@@ -225,7 +225,7 @@ int FPyWrapperArray::Init(FPyWrapperArray* InSelf, const FPyWrapperOwnerContext&
 
 	case EPyConversionMethod::Reference:
 		{
-			PropToUse = InProp;
+			PropToUse = PyUtil::FConstArrayPropOnScope::ExternalReference(InProp);
 			ArrayInstanceToUse = InValue;
 		}
 		break;
@@ -238,7 +238,7 @@ int FPyWrapperArray::Init(FPyWrapperArray* InSelf, const FPyWrapperOwnerContext&
 	check(PropToUse && ArrayInstanceToUse);
 
 	InSelf->OwnerContext = InOwnerContext;
-	InSelf->ArrayProp = PropToUse;
+	InSelf->ArrayProp = MoveTemp(PropToUse);
 	InSelf->ArrayInstance = ArrayInstanceToUse;
 
 	FPyWrapperArrayFactory::Get().MapInstance(InSelf->ArrayInstance, InSelf);
@@ -264,7 +264,7 @@ void FPyWrapperArray::Deinit(FPyWrapperArray* InSelf)
 		}
 		FMemory::Free(InSelf->ArrayInstance);
 	}
-	InSelf->ArrayProp = nullptr;
+	InSelf->ArrayProp.Reset();
 	InSelf->ArrayInstance = nullptr;
 }
 
@@ -1535,7 +1535,7 @@ void FPyWrapperArrayMetaData::AddReferencedObjects(FPyWrapperBase* Instance, FRe
 	FPyWrapperArray* Self = static_cast<FPyWrapperArray*>(Instance);
 	if (Self->ArrayProp && Self->ArrayInstance && !Self->OwnerContext.HasOwner())
 	{
-		const_cast<FArrayProperty*>(Self->ArrayProp)->AddReferencedObjects(Collector);
+		Self->ArrayProp.AddReferencedObjects(Collector);
 		FPyReferenceCollector::AddReferencedObjectsFromProperty(Collector, Self->ArrayProp, Self->ArrayInstance);
 	}
 }
diff --git a/Engine/Plugins/Experimental/PythonScriptPlugin/Source/PythonScriptPlugin/Private/PyWrapperArray.h b/Engine/Plugins/Experimental/PythonScriptPlugin/Source/PythonScriptPlugin/Private/PyWrapperArray.h
--- a/Engine/Plugins/Experimental/PythonScriptPlugin/Source/PythonScriptPlugin/Private/PyWrapperArray.h
+++ b/Engine/Plugins/Experimental/PythonScriptPlugin/Source/PythonScriptPlugin/Private/PyWrapperArray.h
@@ -21,7 +21,7 @@ struct FPyWrapperArray : public FPyWrapperBase
 	FPyWrapperOwnerContext OwnerContext;
 
 	/** Property describing the array */
-	const FArrayProperty* ArrayProp;
+	PyUtil::FConstArrayPropOnScope ArrayProp;
 
 	/** Wrapped array instance */
 	void* ArrayInstance;
diff --git a/Engine/Plugins/Experimental/PythonScriptPlugin/Source/PythonScriptPlugin/Private/PyWrapperDelegate.cpp b/Engine/Plugins/Experimental/PythonScriptPlugin/Source/PythonScriptPlugin/Private/PyWrapperDelegate.cpp
--- a/Engine/Plugins/Experimental/PythonScriptPlugin/Source/PythonScriptPlugin/Private/PyWrapperDelegate.cpp
+++ b/Engine/Plugins/Experimental/PythonScriptPlugin/Source/PythonScriptPlugin/Private/PyWrapperDelegate.cpp
@@ -360,7 +360,7 @@ struct TPyWrapperDelegateImpl
 
 		const PyGenUtil::FGeneratedWrappedFunction& DelegateSignature = MetaDataType::GetDelegateSignature(InSelf);
 
-		if (DelegateSignature.Func->Children == nullptr)
+		if (DelegateSignature.Func->ChildProperties == nullptr)
 		{
 			// Simple case, no parameters or return value
 			FDelegateInvocation::Call(*InSelf->DelegateInstance, nullptr);
diff --git a/Engine/Plugins/Experimental/PythonScriptPlugin/Source/PythonScriptPlugin/Private/PyWrapperFixedArray.cpp b/Engine/Plugins/Experimental/PythonScriptPlugin/Source/PythonScriptPlugin/Private/PyWrapperFixedArray.cpp
--- a/Engine/Plugins/Experimental/PythonScriptPlugin/Source/PythonScriptPlugin/Private/PyWrapperFixedArray.cpp
+++ b/Engine/Plugins/Experimental/PythonScriptPlugin/Source/PythonScriptPlugin/Private/PyWrapperFixedArray.cpp
@@ -2,7 +2,6 @@
 
 #include "PyWrapperFixedArray.h"
 #include "PyWrapperTypeRegistry.h"
-#include "PyUtil.h"
 #include "PyConversion.h"
 #include "PyReferenceCollector.h"
 #include "UObject/UnrealType.h"
@@ -120,7 +119,7 @@ FPyWrapperFixedArray* FPyWrapperFixedArray::New(PyTypeObject* InType)
 	if (Self)
 	{
 		new(&Self->OwnerContext) FPyWrapperOwnerContext();
-		Self->ArrayProp = nullptr;
+		new(&Self->ArrayProp) PyUtil::FConstPropOnScope();
 		Self->ArrayInstance = nullptr;
 	}
 	return Self;
@@ -131,6 +130,7 @@ void FPyWrapperFixedArray::Free(FPyWrapperFixedArray* InSelf)
 	Deinit(InSelf);
 
 	InSelf->OwnerContext.~FPyWrapperOwnerContext();
+	InSelf->ArrayProp.~TPropOnScope();
 	FPyWrapperBase::Free(InSelf);
 }
 
@@ -144,7 +144,7 @@ int FPyWrapperFixedArray::Init(FPyWrapperFixedArray* InSelf, const PyUtil::FProp
 		return BaseInit;
 	}
 
-	FProperty* ArrayProp = PyUtil::CreateProperty(InPropDef, InLen);
+	PyUtil::FConstPropOnScope ArrayProp = PyUtil::FConstPropOnScope::OwnedReference(PyUtil::CreateProperty(InPropDef, InLen));
 	if (!ArrayProp)
 	{
 		PyUtil::SetPythonError(PyExc_Exception, InSelf, TEXT("Array property was null during init"));
@@ -154,7 +154,7 @@ int FPyWrapperFixedArray::Init(FPyWrapperFixedArray* InSelf, const PyUtil::FProp
 	void* ArrayValue = FMemory::Malloc(ArrayProp->GetSize(), ArrayProp->GetMinAlignment());
 	ArrayProp->InitializeValue(ArrayValue);
 
-	InSelf->ArrayProp = ArrayProp;
+	InSelf->ArrayProp = MoveTemp(ArrayProp);
 	InSelf->ArrayInstance = ArrayValue;
 
 	FPyWrapperFixedArrayFactory::Get().MapInstance(InSelf->ArrayInstance, InSelf);
@@ -175,14 +175,14 @@ int FPyWrapperFixedArray::Init(FPyWrapperFixedArray* InSelf, const FPyWrapperOwn
 
 	check(InProp && InValue);
 
-	const FProperty* PropToUse = nullptr;
+	PyUtil::FConstPropOnScope PropToUse;
 	void* ArrayInstanceToUse = nullptr;
 	switch (InConversionMethod)
 	{
 	case EPyConversionMethod::Copy:
 	case EPyConversionMethod::Steal:
 		{
-			PropToUse = PyUtil::CreateProperty(InProp);
+			PropToUse = PyUtil::FConstPropOnScope::OwnedReference(PyUtil::CreateProperty(InProp));
 			if (!PropToUse)
 			{
 				PyUtil::SetPythonError(PyExc_TypeError, InSelf, *FString::Printf(TEXT("Failed to create property from '%s' (%s)"), *InProp->GetName(), *InProp->GetClass()->GetName()));
@@ -197,7 +197,7 @@ int FPyWrapperFixedArray::Init(FPyWrapperFixedArray* InSelf, const FPyWrapperOwn
 
 	case EPyConversionMethod::Reference:
 		{
-			PropToUse = InProp;
+			PropToUse = PyUtil::FConstPropOnScope::ExternalReference(InProp);
 			ArrayInstanceToUse = InValue;
 		}
 		break;
@@ -210,7 +210,7 @@ int FPyWrapperFixedArray::Init(FPyWrapperFixedArray* InSelf, const FPyWrapperOwn
 	check(PropToUse && ArrayInstanceToUse);
 
 	InSelf->OwnerContext = InOwnerContext;
-	InSelf->ArrayProp = PropToUse;
+	InSelf->ArrayProp = MoveTemp(PropToUse);
 	InSelf->ArrayInstance = ArrayInstanceToUse;
 
 	FPyWrapperFixedArrayFactory::Get().MapInstance(InSelf->ArrayInstance, InSelf);
@@ -236,7 +236,7 @@ void FPyWrapperFixedArray::Deinit(FPyWrapperFixedArray* InSelf)
 		}
 		FMemory::Free(InSelf->ArrayInstance);
 	}
-	InSelf->ArrayProp = nullptr;
+	InSelf->ArrayProp.Release();
 	InSelf->ArrayInstance = nullptr;
 }
 
@@ -284,7 +284,7 @@ FPyWrapperFixedArray* FPyWrapperFixedArray::CastPyObject(PyObject* InPyObject, P
 			return nullptr;
 		}
 
-		const PyUtil::FPropertyDef SelfPropDef = Self->ArrayProp;
+		const PyUtil::FPropertyDef SelfPropDef = Self->ArrayProp.Get();
 		if (SelfPropDef == InPropDef)
 		{
 			SetOptionalPyConversionResult(FPyConversionResult::Success(), OutCastResult);
@@ -460,7 +460,7 @@ FPyWrapperFixedArray* FPyWrapperFixedArray::Concat(FPyWrapperFixedArray* InSelf,
 		return nullptr;
 	}
 
-	const PyUtil::FPropertyDef SelfPropDef = InSelf->ArrayProp;
+	const PyUtil::FPropertyDef SelfPropDef = InSelf->ArrayProp.Get();
 	FPyWrapperFixedArrayPtr Other = FPyWrapperFixedArrayPtr::StealReference(FPyWrapperFixedArray::CastPyObject(InOther, &PyWrapperFixedArrayType, SelfPropDef));
 	if (!Other)
 	{
@@ -500,7 +500,7 @@ FPyWrapperFixedArray* FPyWrapperFixedArray::Repeat(FPyWrapperFixedArray* InSelf,
 		return nullptr;
 	}
 
-	const PyUtil::FPropertyDef SelfPropDef = InSelf->ArrayProp;
+	const PyUtil::FPropertyDef SelfPropDef = InSelf->ArrayProp.Get();
 	const int32 NewArrayLen = InSelf->ArrayProp->ArrayDim * InMultiplier;
 	FPyWrapperFixedArrayPtr NewArray = FPyWrapperFixedArrayPtr::StealReference(FPyWrapperFixedArray::New(Py_TYPE(InSelf)));
 	if (FPyWrapperFixedArray::Init(NewArray, SelfPropDef, NewArrayLen) != 0)
@@ -588,7 +588,7 @@ PyTypeObject InitializePyWrapperFixedArrayType()
 				return nullptr;
 			}
 
-			const PyUtil::FPropertyDef SelfPropDef = InSelf->ArrayProp;
+			const PyUtil::FPropertyDef SelfPropDef = InSelf->ArrayProp.Get();
 			FPyWrapperFixedArrayPtr Other = FPyWrapperFixedArrayPtr::StealReference(FPyWrapperFixedArray::CastPyObject(InOther, &PyWrapperFixedArrayType, SelfPropDef));
 			if (!Other)
 			{
@@ -699,7 +699,7 @@ PyTypeObject InitializePyWrapperFixedArrayType()
 				return nullptr;
 			}
 
-			const PyUtil::FPropertyDef SelfPropDef = InSelf->ArrayProp;
+			const PyUtil::FPropertyDef SelfPropDef = InSelf->ArrayProp.Get();
 			FPyWrapperFixedArrayPtr NewArray = FPyWrapperFixedArrayPtr::StealReference(FPyWrapperFixedArray::New(Py_TYPE(InSelf)));
 			if (FPyWrapperFixedArray::Init(NewArray, SelfPropDef, SliceLen) != 0)
 			{
@@ -906,7 +906,7 @@ void FPyWrapperFixedArrayMetaData::AddReferencedObjects(FPyWrapperBase* Instance
 	FPyWrapperFixedArray* Self = static_cast<FPyWrapperFixedArray*>(Instance);
 	if (Self->ArrayProp && Self->ArrayInstance && !Self->OwnerContext.HasOwner())
 	{
-		Collector.AddReferencedObject(Self->ArrayProp);
+		Self->ArrayProp.AddReferencedObjects(Collector);
 		FPyReferenceCollector::AddReferencedObjectsFromProperty(Collector, Self->ArrayProp, Self->ArrayInstance);
 	}
 }
diff --git a/Engine/Plugins/Experimental/PythonScriptPlugin/Source/PythonScriptPlugin/Private/PyWrapperFixedArray.h b/Engine/Plugins/Experimental/PythonScriptPlugin/Source/PythonScriptPlugin/Private/PyWrapperFixedArray.h
--- a/Engine/Plugins/Experimental/PythonScriptPlugin/Source/PythonScriptPlugin/Private/PyWrapperFixedArray.h
+++ b/Engine/Plugins/Experimental/PythonScriptPlugin/Source/PythonScriptPlugin/Private/PyWrapperFixedArray.h
@@ -21,7 +21,7 @@ struct FPyWrapperFixedArray : public FPyWrapperBase
 	FPyWrapperOwnerContext OwnerContext;
 
 	/** Property describing the fixed-array */
-	const FProperty* ArrayProp;
+	PyUtil::FConstPropOnScope ArrayProp;
 
 	/** Wrapped fixed-array instance */
 	void* ArrayInstance;
diff --git a/Engine/Plugins/Experimental/PythonScriptPlugin/Source/PythonScriptPlugin/Private/PyWrapperMap.cpp b/Engine/Plugins/Experimental/PythonScriptPlugin/Source/PythonScriptPlugin/Private/PyWrapperMap.cpp
--- a/Engine/Plugins/Experimental/PythonScriptPlugin/Source/PythonScriptPlugin/Private/PyWrapperMap.cpp
+++ b/Engine/Plugins/Experimental/PythonScriptPlugin/Source/PythonScriptPlugin/Private/PyWrapperMap.cpp
@@ -2,8 +2,6 @@
 
 #include "PyWrapperMap.h"
 #include "PyWrapperTypeRegistry.h"
-#include "PyCore.h"
-#include "PyUtil.h"
 #include "PyConversion.h"
 #include "PyReferenceCollector.h"
 #include "UObject/UnrealType.h"
@@ -435,7 +433,7 @@ FPyWrapperMap* FPyWrapperMap::New(PyTypeObject* InType)
 	if (Self)
 	{
 		new(&Self->OwnerContext) FPyWrapperOwnerContext();
-		Self->MapProp = nullptr;
+		new(&Self->MapProp) PyUtil::FConstMapPropOnScope();
 		Self->MapInstance = nullptr;
 	}
 	return Self;
@@ -446,6 +444,7 @@ void FPyWrapperMap::Free(FPyWrapperMap* InSelf)
 	Deinit(InSelf);
 
 	InSelf->OwnerContext.~FPyWrapperOwnerContext();
+	InSelf->MapProp.~TPropOnScope();
 	FPyWrapperBase::Free(InSelf);
 }
 
@@ -459,23 +458,23 @@ int FPyWrapperMap::Init(FPyWrapperMap* InSelf, const PyUtil::FPropertyDef& InKey
 		return BaseInit;
 	}
 
-	FProperty* MapKeyProp = PyUtil::CreateProperty(InKeyDef, 1);
+	PyUtil::FPropOnScope MapKeyProp = PyUtil::FPropOnScope::OwnedReference(PyUtil::CreateProperty(InKeyDef, 1));
 	if (!MapKeyProp)
 	{
 		PyUtil::SetPythonError(PyExc_Exception, InSelf, TEXT("Map key property was null during init"));
 		return -1;
 	}
 
-	FProperty* MapValueProp = PyUtil::CreateProperty(InValueDef, 1);
+	PyUtil::FPropOnScope MapValueProp = PyUtil::FPropOnScope::OwnedReference(PyUtil::CreateProperty(InValueDef, 1));
 	if (!MapValueProp)
 	{
 		PyUtil::SetPythonError(PyExc_Exception, InSelf, TEXT("Map value property was null during init"));
 		return -1;
 	}
 
-	FMapProperty* MapProp = new FMapProperty(GetPythonPropertyContainer(), FField::GenerateFFieldName(GetPythonPropertyContainer(), FMapProperty::StaticClass()), RF_NoFlags);
-	MapProp->KeyProp = MapKeyProp;
-	MapProp->ValueProp = MapValueProp;
+	PyUtil::FMapPropOnScope MapProp = PyUtil::FMapPropOnScope::OwnedReference(new FMapProperty(FFieldVariant(), PyUtil::DefaultPythonPropertyName, RF_NoFlags));
+	MapProp->KeyProp = MapKeyProp.Release();
+	MapProp->ValueProp = MapValueProp.Release();
 
 	// Need to manually call Link to fix-up some data (such as the C++ property flags and the map layout) that are only set during Link
 	{
@@ -486,7 +485,7 @@ int FPyWrapperMap::Init(FPyWrapperMap* InSelf, const PyUtil::FPropertyDef& InKey
 	void* MapValue = FMemory::Malloc(MapProp->GetSize(), MapProp->GetMinAlignment());
 	MapProp->InitializeValue(MapValue);
 
-	InSelf->MapProp = MapProp;
+	InSelf->MapProp = MoveTemp(MapProp);
 	InSelf->MapInstance = MapValue;
 
 	FPyWrapperMapFactory::Get().MapInstance(InSelf->MapInstance, InSelf);
@@ -507,31 +506,30 @@ int FPyWrapperMap::Init(FPyWrapperMap* InSelf, const FPyWrapperOwnerContext& InO
 
 	check(InProp && InValue);
 
-	const FMapProperty* PropToUse = nullptr;
+	PyUtil::FConstMapPropOnScope PropToUse;
 	void* MapInstanceToUse = nullptr;
 	switch (InConversionMethod)
 	{
 	case EPyConversionMethod::Copy:
 	case EPyConversionMethod::Steal:
 		{
-			FProperty* MapKeyProp = PyUtil::CreateProperty(InProp->KeyProp, 1);
+			PyUtil::FPropOnScope MapKeyProp = PyUtil::FPropOnScope::OwnedReference(PyUtil::CreateProperty(InProp->KeyProp, 1));
 			if (!MapKeyProp)
 			{
 				PyUtil::SetPythonError(PyExc_TypeError, InSelf, *FString::Printf(TEXT("Failed to create key property from '%s' (%s)"), *InProp->KeyProp->GetName(), *InProp->KeyProp->GetClass()->GetName()));
 				return -1;
 			}
 
-			FProperty* MapValueProp = PyUtil::CreateProperty(InProp->ValueProp, 1);
+			PyUtil::FPropOnScope MapValueProp = PyUtil::FPropOnScope::OwnedReference(PyUtil::CreateProperty(InProp->ValueProp, 1));
 			if (!MapValueProp)
 			{
 				PyUtil::SetPythonError(PyExc_TypeError, InSelf, *FString::Printf(TEXT("Failed to create value property from '%s' (%s)"), *InProp->ValueProp->GetName(), *InProp->ValueProp->GetClass()->GetName()));
 				return -1;
 			}
 
-			FMapProperty* MapProp = new FMapProperty(GetPythonPropertyContainer(), FField::GenerateFFieldName(GetPythonPropertyContainer(), FMapProperty::StaticClass()), RF_NoFlags);
-			MapProp->KeyProp = MapKeyProp;
-			MapProp->ValueProp = MapValueProp;
-			PropToUse = MapProp;
+			PyUtil::FMapPropOnScope MapProp = PyUtil::FMapPropOnScope::OwnedReference(new FMapProperty(FFieldVariant(), PyUtil::DefaultPythonPropertyName, RF_NoFlags));
+			MapProp->KeyProp = MapKeyProp.Release();
+			MapProp->ValueProp = MapValueProp.Release();
 
 			// Need to manually call Link to fix-up some data (such as the C++ property flags and the map layout) that are only set during Link
 			{
@@ -539,6 +537,8 @@ int FPyWrapperMap::Init(FPyWrapperMap* InSelf, const FPyWrapperOwnerContext& InO
 				MapProp->LinkWithoutChangingOffset(Ar);
 			}
 
+			PropToUse = MoveTemp(MapProp);
+
 			MapInstanceToUse = FMemory::Malloc(PropToUse->GetSize(), PropToUse->GetMinAlignment());
 			PropToUse->InitializeValue(MapInstanceToUse);
 			if (InConversionMethod == EPyConversionMethod::Steal)
@@ -555,7 +555,7 @@ int FPyWrapperMap::Init(FPyWrapperMap* InSelf, const FPyWrapperOwnerContext& InO
 
 	case EPyConversionMethod::Reference:
 		{
-			PropToUse = InProp;
+			PropToUse = PyUtil::FConstMapPropOnScope::ExternalReference(InProp);
 			MapInstanceToUse = InValue;
 		}
 		break;
@@ -568,7 +568,7 @@ int FPyWrapperMap::Init(FPyWrapperMap* InSelf, const FPyWrapperOwnerContext& InO
 	check(PropToUse && MapInstanceToUse);
 
 	InSelf->OwnerContext = InOwnerContext;
-	InSelf->MapProp = PropToUse;
+	InSelf->MapProp = MoveTemp(PropToUse);
 	InSelf->MapInstance = MapInstanceToUse;
 
 	FPyWrapperMapFactory::Get().MapInstance(InSelf->MapInstance, InSelf);
@@ -594,7 +594,7 @@ void FPyWrapperMap::Deinit(FPyWrapperMap* InSelf)
 		}
 		FMemory::Free(InSelf->MapInstance);
 	}
-	InSelf->MapProp = nullptr;
+	InSelf->MapProp.Reset();
 	InSelf->MapInstance = nullptr;
 }
 
@@ -1852,7 +1852,7 @@ void FPyWrapperMapMetaData::AddReferencedObjects(FPyWrapperBase* Instance, FRefe
 	FPyWrapperMap* Self = static_cast<FPyWrapperMap*>(Instance);
 	if (Self->MapProp && Self->MapInstance && !Self->OwnerContext.HasOwner())
 	{
-		Collector.AddReferencedObject(Self->MapProp);
+		Self->MapProp.AddReferencedObjects(Collector);
 		FPyReferenceCollector::AddReferencedObjectsFromProperty(Collector, Self->MapProp, Self->MapInstance);
 	}
 }
diff --git a/Engine/Plugins/Experimental/PythonScriptPlugin/Source/PythonScriptPlugin/Private/PyWrapperMap.h b/Engine/Plugins/Experimental/PythonScriptPlugin/Source/PythonScriptPlugin/Private/PyWrapperMap.h
--- a/Engine/Plugins/Experimental/PythonScriptPlugin/Source/PythonScriptPlugin/Private/PyWrapperMap.h
+++ b/Engine/Plugins/Experimental/PythonScriptPlugin/Source/PythonScriptPlugin/Private/PyWrapperMap.h
@@ -21,7 +21,7 @@ struct FPyWrapperMap : public FPyWrapperBase
 	FPyWrapperOwnerContext OwnerContext;
 
 	/** Property describing the map */
-	const FMapProperty* MapProp;
+	PyUtil::FConstMapPropOnScope MapProp;
 
 	/** Wrapped map instance */
 	void* MapInstance;
diff --git a/Engine/Plugins/Experimental/PythonScriptPlugin/Source/PythonScriptPlugin/Private/PyWrapperObject.cpp b/Engine/Plugins/Experimental/PythonScriptPlugin/Source/PythonScriptPlugin/Private/PyWrapperObject.cpp
--- a/Engine/Plugins/Experimental/PythonScriptPlugin/Source/PythonScriptPlugin/Private/PyWrapperObject.cpp
+++ b/Engine/Plugins/Experimental/PythonScriptPlugin/Source/PythonScriptPlugin/Private/PyWrapperObject.cpp
@@ -1460,7 +1460,7 @@ public:
 			if (FuncArgDefaults.IsValidIndex(InputArgIndex) && FuncArgDefaults[InputArgIndex])
 			{
 				// Convert the default value to the given property...
-				PyUtil::FPropValueOnScope DefaultValue(Param);
+				PyUtil::FPropValueOnScope DefaultValue(PyUtil::FConstPropOnScope::ExternalReference(Param));
 				if (!DefaultValue.IsValid() || !DefaultValue.SetValue(FuncArgDefaults[InputArgIndex], *PyUtil::GetErrorContext(PyType)))
 				{
 					PyUtil::SetPythonError(PyExc_Exception, PyType, *FString::Printf(TEXT("Failed to convert default value for function '%s' argument '%s' (%s)"), *InFieldName, *FuncArgNames[InputArgIndex], *Param->GetClass()->GetName()));
diff --git a/Engine/Plugins/Experimental/PythonScriptPlugin/Source/PythonScriptPlugin/Private/PyWrapperSet.cpp b/Engine/Plugins/Experimental/PythonScriptPlugin/Source/PythonScriptPlugin/Private/PyWrapperSet.cpp
--- a/Engine/Plugins/Experimental/PythonScriptPlugin/Source/PythonScriptPlugin/Private/PyWrapperSet.cpp
+++ b/Engine/Plugins/Experimental/PythonScriptPlugin/Source/PythonScriptPlugin/Private/PyWrapperSet.cpp
@@ -2,8 +2,6 @@
 
 #include "PyWrapperSet.h"
 #include "PyWrapperTypeRegistry.h"
-#include "PyCore.h"
-#include "PyUtil.h"
 #include "PyConversion.h"
 #include "PyReferenceCollector.h"
 #include "UObject/UnrealType.h"
@@ -157,7 +155,7 @@ FPyWrapperSet* FPyWrapperSet::New(PyTypeObject* InType)
 	if (Self)
 	{
 		new(&Self->OwnerContext) FPyWrapperOwnerContext();
-		Self->SetProp = nullptr;
+		new(&Self->SetProp) PyUtil::FConstSetPropOnScope();
 		Self->SetInstance = nullptr;
 	}
 	return Self;
@@ -168,6 +166,7 @@ void FPyWrapperSet::Free(FPyWrapperSet* InSelf)
 	Deinit(InSelf);
 
 	InSelf->OwnerContext.~FPyWrapperOwnerContext();
+	InSelf->SetProp.~TPropOnScope();
 	FPyWrapperBase::Free(InSelf);
 }
 
@@ -181,15 +180,15 @@ int FPyWrapperSet::Init(FPyWrapperSet* InSelf, const PyUtil::FPropertyDef& InEle
 		return BaseInit;
 	}
 
-	FProperty* SetElementProp = PyUtil::CreateProperty(InElementDef, 1);
+	PyUtil::FPropOnScope SetElementProp = PyUtil::FPropOnScope::OwnedReference(PyUtil::CreateProperty(InElementDef, 1));
 	if (!SetElementProp)
 	{
 		PyUtil::SetPythonError(PyExc_Exception, InSelf, TEXT("Set element property was null during init"));
 		return -1;
 	}
 
-	FSetProperty* SetProp = new FSetProperty(GetPythonPropertyContainer(), FField::GenerateFFieldName(GetPythonPropertyContainer(), FSetProperty::StaticClass()), RF_NoFlags);
-	SetProp->ElementProp = SetElementProp;
+	PyUtil::FSetPropOnScope SetProp = PyUtil::FSetPropOnScope::OwnedReference(new FSetProperty(FFieldVariant(), PyUtil::DefaultPythonPropertyName, RF_NoFlags));
+	SetProp->ElementProp = SetElementProp.Release();
 
 	// Need to manually call Link to fix-up some data (such as the C++ property flags and the set layout) that are only set during Link
 	{
@@ -200,7 +199,7 @@ int FPyWrapperSet::Init(FPyWrapperSet* InSelf, const PyUtil::FPropertyDef& InEle
 	void* SetValue = FMemory::Malloc(SetProp->GetSize(), SetProp->GetMinAlignment());
 	SetProp->InitializeValue(SetValue);
 
-	InSelf->SetProp = SetProp;
+	InSelf->SetProp = MoveTemp(SetProp);
 	InSelf->SetInstance = SetValue;
 
 	if (!InSelf->SetProp->ElementProp->HasAnyPropertyFlags(CPF_HasGetValueTypeHash))
@@ -227,23 +226,22 @@ int FPyWrapperSet::Init(FPyWrapperSet* InSelf, const FPyWrapperOwnerContext& InO
 
 	check(InProp && InValue);
 
-	const FSetProperty* PropToUse = nullptr;
+	PyUtil::FConstSetPropOnScope PropToUse;
 	void* SetInstanceToUse = nullptr;
 	switch (InConversionMethod)
 	{
 	case EPyConversionMethod::Copy:
 	case EPyConversionMethod::Steal:
 		{
-			FProperty* SetElementProp = PyUtil::CreateProperty(InProp->ElementProp);
+			PyUtil::FPropOnScope SetElementProp = PyUtil::FPropOnScope::OwnedReference(PyUtil::CreateProperty(InProp->ElementProp));
 			if (!SetElementProp)
 			{
 				PyUtil::SetPythonError(PyExc_TypeError, InSelf, *FString::Printf(TEXT("Failed to create element property from '%s' (%s)"), *InProp->ElementProp->GetName(), *InProp->ElementProp->GetClass()->GetName()));
 				return -1;
 			}
 
-			FSetProperty* SetProp = new FSetProperty(GetPythonPropertyContainer(), FField::GenerateFFieldName(GetPythonPropertyContainer(), FSetProperty::StaticClass()), RF_NoFlags);
-			SetProp->ElementProp = SetElementProp;
-			PropToUse = SetProp;
+			PyUtil::FSetPropOnScope SetProp = PyUtil::FSetPropOnScope::OwnedReference(new FSetProperty(FFieldVariant(), PyUtil::DefaultPythonPropertyName, RF_NoFlags));
+			SetProp->ElementProp = SetElementProp.Release();
 			
 			// Need to manually call Link to fix-up some data (such as the C++ property flags and the set layout) that are only set during Link
 			{
@@ -251,6 +249,8 @@ int FPyWrapperSet::Init(FPyWrapperSet* InSelf, const FPyWrapperOwnerContext& InO
 				SetProp->LinkWithoutChangingOffset(Ar);
 			}
 
+			PropToUse = MoveTemp(SetProp);
+
 			SetInstanceToUse = FMemory::Malloc(PropToUse->GetSize(), PropToUse->GetMinAlignment());
 			PropToUse->InitializeValue(SetInstanceToUse);
 			if (InConversionMethod == EPyConversionMethod::Steal)
@@ -267,7 +267,7 @@ int FPyWrapperSet::Init(FPyWrapperSet* InSelf, const FPyWrapperOwnerContext& InO
 
 	case EPyConversionMethod::Reference:
 		{
-			PropToUse = InProp;
+			PropToUse = PyUtil::FConstSetPropOnScope::ExternalReference(InProp);
 			SetInstanceToUse = InValue;
 		}
 		break;
@@ -280,7 +280,7 @@ int FPyWrapperSet::Init(FPyWrapperSet* InSelf, const FPyWrapperOwnerContext& InO
 	check(PropToUse && SetInstanceToUse);
 
 	InSelf->OwnerContext = InOwnerContext;
-	InSelf->SetProp = PropToUse;
+	InSelf->SetProp = MoveTemp(PropToUse);
 	InSelf->SetInstance = SetInstanceToUse;
 
 	if (!InSelf->SetProp->ElementProp->HasAnyPropertyFlags(CPF_HasGetValueTypeHash))
@@ -312,7 +312,7 @@ void FPyWrapperSet::Deinit(FPyWrapperSet* InSelf)
 		}
 		FMemory::Free(InSelf->SetInstance);
 	}
-	InSelf->SetProp = nullptr;
+	InSelf->SetProp.Reset();
 	InSelf->SetInstance = nullptr;
 }
 
@@ -1608,7 +1608,7 @@ void FPyWrapperSetMetaData::AddReferencedObjects(FPyWrapperBase* Instance, FRefe
 	FPyWrapperSet* Self = static_cast<FPyWrapperSet*>(Instance);
 	if (Self->SetProp && Self->SetInstance && !Self->OwnerContext.HasOwner())
 	{
-		Collector.AddReferencedObject(Self->SetProp);
+		Self->SetProp.AddReferencedObjects(Collector);
 		FPyReferenceCollector::AddReferencedObjectsFromProperty(Collector, Self->SetProp, Self->SetInstance);
 	}
 }
diff --git a/Engine/Plugins/Experimental/PythonScriptPlugin/Source/PythonScriptPlugin/Private/PyWrapperSet.h b/Engine/Plugins/Experimental/PythonScriptPlugin/Source/PythonScriptPlugin/Private/PyWrapperSet.h
--- a/Engine/Plugins/Experimental/PythonScriptPlugin/Source/PythonScriptPlugin/Private/PyWrapperSet.h
+++ b/Engine/Plugins/Experimental/PythonScriptPlugin/Source/PythonScriptPlugin/Private/PyWrapperSet.h
@@ -21,7 +21,7 @@ struct FPyWrapperSet : public FPyWrapperBase
 	FPyWrapperOwnerContext OwnerContext;
 
 	/** Property describing the set */
-	const FSetProperty* SetProp;
+	PyUtil::FConstSetPropOnScope SetProp;
 
 	/** Wrapped set instance */
 	void* SetInstance;
