Ensure start-up scripts run when using Python in commandlet mode

diff --git a/Engine/Plugins/Experimental/PythonScriptPlugin/Source/PythonScriptPlugin/Private/PythonScriptCommandlet.cpp b/Engine/Plugins/Experimental/PythonScriptPlugin/Source/PythonScriptPlugin/Private/PythonScriptCommandlet.cpp
--- a/Engine/Plugins/Experimental/PythonScriptPlugin/Source/PythonScriptPlugin/Private/PythonScriptCommandlet.cpp
+++ b/Engine/Plugins/Experimental/PythonScriptPlugin/Source/PythonScriptPlugin/Private/PythonScriptCommandlet.cpp
@@ -3,6 +3,7 @@
 #include "PythonScriptCommandlet.h"
 #include "IPythonScriptPlugin.h"
 #include "Logging/LogMacros.h"
+#include "Containers/Ticker.h"
 
 DEFINE_LOG_CATEGORY_STATIC(LogPythonScriptCommandlet, Log, All);
 
@@ -34,6 +35,9 @@ int32 UPythonScriptCommandlet::Main(const FString& Params)
 
 #if WITH_PYTHON
 	{
+		// Tick once to ensure that any start-up scripts have been run
+		FTicker::GetCoreTicker().Tick(0.0f);
+
 		UE_LOG(LogPythonScriptCommandlet, Display, TEXT("Running Python script: %s"), *PythonScript);
 
 		FPythonCommandEx PythonCommand;
