Expose functions to register and unregister callbacks for python shutdown

diff --git a/Engine/Plugins/Experimental/PythonScriptPlugin/Source/PythonScriptPlugin/Private/PyCore.cpp b/Engine/Plugins/Experimental/PythonScriptPlugin/Source/PythonScriptPlugin/Private/PyCore.cpp
--- a/Engine/Plugins/Experimental/PythonScriptPlugin/Source/PythonScriptPlugin/Private/PyCore.cpp
+++ b/Engine/Plugins/Experimental/PythonScriptPlugin/Source/PythonScriptPlugin/Private/PyCore.cpp
@@ -1651,6 +1651,52 @@ PyObject* CreateLocalizedTextFromStringTable(PyObject* InSelf, PyObject* InArgs)
 	return PyConversion::Pythonize(FText::FromStringTable(Id, Key));
 }
 
+PyObject* RegisterPythonShutdownCallback(PyObject* InSelf, PyObject* InArgs)
+{
+	PyObject* PyObj = nullptr;
+	if (!PyArg_ParseTuple(InArgs, "O:register_python_shutdown_callback", &PyObj))
+	{
+		return nullptr;
+	}
+	check(PyObj);
+
+	FPyObjectPtr PyCallable = FPyObjectPtr::NewReference(PyObj);
+	FDelegateHandle PythonShutdownDelegateHandle = FPythonScriptPlugin::Get()->OnPythonShutdown().AddLambda([PyCallable]() mutable
+	{
+		FPyObjectPtr Result = FPyObjectPtr::StealReference(PyObject_CallObject(PyCallable.GetPtr(), nullptr));
+		if (!Result)
+		{
+			PyUtil::LogPythonError();
+		}
+	});
+
+	return (PyObject*)FPyDelegateHandle::CreateInstance(PythonShutdownDelegateHandle);
+}
+
+PyObject* UnregisterPythonShutdownCallback(PyObject* InSelf, PyObject* InArgs)
+{
+	PyObject* PyObj = nullptr;
+	if (!PyArg_ParseTuple(InArgs, "O:unregister_python_shutdown_callback", &PyObj))
+	{
+		return nullptr;
+	}
+	check(PyObj);
+
+	FPyDelegateHandlePtr PythonEventDelegateHandle = FPyDelegateHandlePtr::StealReference(FPyDelegateHandle::CastPyObject(PyObj));
+	if (!PythonEventDelegateHandle)
+	{
+		PyUtil::SetPythonError(PyExc_TypeError, TEXT("unregister_python_shutdown_callback"), *FString::Printf(TEXT("Failed to convert argument '%s' to '_DelegateHandle'"), *PyUtil::GetFriendlyTypename(PyObj)));
+		return nullptr;
+	}
+
+	if (PythonEventDelegateHandle->Value.IsValid())
+	{
+		FPythonScriptPlugin::Get()->OnPythonShutdown().Remove(PythonEventDelegateHandle->Value);
+	}
+
+	Py_RETURN_NONE;
+}
+
 PyMethodDef PyCoreMethods[] = {
 	{ "log", PyCFunctionCast(&Log), METH_VARARGS, "x.log(str) -> None -- log the given argument as information in the LogPython category" },
 	{ "log_warning", PyCFunctionCast(&LogWarning), METH_VARARGS, "x.log_warning(str) -> None -- log the given argument as a warning in the LogPython category" },
@@ -1674,6 +1720,8 @@ PyMethodDef PyCoreMethods[] = {
 	{ "get_type_from_class", PyCFunctionCast(&GetTypeFromClass), METH_VARARGS, "x.get_type_from_class(class) -> type -- get the best matching Python type for the given Unreal class" },
 	{ "get_type_from_struct", PyCFunctionCast(&GetTypeFromStruct), METH_VARARGS, "x.get_type_from_struct(struct) -> type -- get the best matching Python type for the given Unreal struct" },
 	{ "get_type_from_enum", PyCFunctionCast(&GetTypeFromEnum), METH_VARARGS, "x.get_type_from_enum(enum) -> type -- get the best matching Python type for the given Unreal enum" },
+	{ "register_python_shutdown_callback", PyCFunctionCast(&RegisterPythonShutdownCallback), METH_VARARGS, "x.register_python_shutdown_callback(callable) -> _DelegateHandle -- register the given callable (with no input arguments) as a callback to execute immediately before Python shutdown"},
+	{ "unregister_python_shutdown_callback", PyCFunctionCast(&UnregisterPythonShutdownCallback), METH_VARARGS, "x.unregister_python_shutdown_callback(handle) -> None -- unregister the given handle from a previous call to register_python_shutdown_callback"},
 	{ "NSLOCTEXT", PyCFunctionCast(&CreateLocalizedText), METH_VARARGS, "x.NSLOCTEXT(ns, key, source) -> Text -- create a localized Text from the given namespace, key, and source string" },
 	{ "LOCTABLE", PyCFunctionCast(&CreateLocalizedTextFromStringTable), METH_VARARGS, "x.LOCTABLE(id, key) -> Text -- get a localized Text from the given string table id and key" },
 	{ nullptr, nullptr, 0, nullptr }
