Fix issue with Python script plugin holding reference into array while it's being modified, leading to heap corruption

diff --git a/Engine/Plugins/Experimental/PythonScriptPlugin/Source/PythonScriptPlugin/Private/PyWrapperTypeRegistry.cpp b/Engine/Plugins/Experimental/PythonScriptPlugin/Source/PythonScriptPlugin/Private/PyWrapperTypeRegistry.cpp
--- a/Engine/Plugins/Experimental/PythonScriptPlugin/Source/PythonScriptPlugin/Private/PyWrapperTypeRegistry.cpp
+++ b/Engine/Plugins/Experimental/PythonScriptPlugin/Source/PythonScriptPlugin/Private/PyWrapperTypeRegistry.cpp
@@ -1092,6 +1092,11 @@ PyTypeObject* FPyWrapperTypeRegistry::GenerateWrappedClassType(const UClass* InC
 			GeneratedWrappedMethod.MethodCallback = GeneratedWrappedMethod.MethodFunc.InputParams.Num() > 0 ? PyCFunctionWithClosureCast(&FPyWrapperObject::CallMethodWithArgs_Impl) : PyCFunctionWithClosureCast(&FPyWrapperObject::CallMethodNoArgs_Impl);
 		}
 
+		// We must create a copy here because otherwise the reference will get invalidated by 
+		// subsequent modifications
+
+		const PyGenUtil::FGeneratedWrappedMethod GeneratedWrappedMethodCopy = GeneratedWrappedMethod;
+
 		const TArray<FString> DeprecatedPythonFuncNames = PyGenUtil::GetDeprecatedFunctionPythonNames(InFunc);
 		for (const FString& DeprecatedPythonFuncName : DeprecatedPythonFuncNames)
 		{
@@ -1099,7 +1104,7 @@ PyTypeObject* FPyWrapperTypeRegistry::GenerateWrappedClassType(const UClass* InC
 			PythonMethods.Add(*DeprecatedPythonFuncName, InFunc->GetFName());
 			PythonDeprecatedMethods.Add(*DeprecatedPythonFuncName, DeprecationMessage);
 
-			PyGenUtil::FGeneratedWrappedMethod DeprecatedGeneratedWrappedMethod = GeneratedWrappedMethod;
+			PyGenUtil::FGeneratedWrappedMethod DeprecatedGeneratedWrappedMethod = GeneratedWrappedMethodCopy;
 			DeprecatedGeneratedWrappedMethod.MethodName = PyGenUtil::TCHARToUTF8Buffer(*DeprecatedPythonFuncName);
 			DeprecatedGeneratedWrappedMethod.MethodDoc = PyGenUtil::TCHARToUTF8Buffer(*FString::Printf(TEXT("deprecated: %s"), *DeprecationMessage));
 			DeprecatedGeneratedWrappedMethod.MethodFunc.DeprecationMessage = MoveTemp(DeprecationMessage);
@@ -1111,11 +1116,11 @@ PyTypeObject* FPyWrapperTypeRegistry::GenerateWrappedClassType(const UClass* InC
 		// Should this function also be hoisted as a struct method or operator?
 		if (InFunc->HasMetaData(PyGenUtil::ScriptMethodMetaDataKey))
 		{
-			GenerateWrappedDynamicMethod(InFunc, GeneratedWrappedMethod);
+			GenerateWrappedDynamicMethod(InFunc, GeneratedWrappedMethodCopy);
 		}
 		if (InFunc->HasMetaData(PyGenUtil::ScriptOperatorMetaDataKey))
 		{
-			GenerateWrappedOperator(InFunc, GeneratedWrappedMethod);
+			GenerateWrappedOperator(InFunc, GeneratedWrappedMethodCopy);
 		}
 	};
 
