Use of TUniqueObj to prevent dangling references from reallocations.

diff --git a/Engine/Plugins/Experimental/PythonScriptPlugin/Source/PythonScriptPlugin/Private/PyWrapperTypeRegistry.cpp b/Engine/Plugins/Experimental/PythonScriptPlugin/Source/PythonScriptPlugin/Private/PyWrapperTypeRegistry.cpp
--- a/Engine/Plugins/Experimental/PythonScriptPlugin/Source/PythonScriptPlugin/Private/PyWrapperTypeRegistry.cpp
+++ b/Engine/Plugins/Experimental/PythonScriptPlugin/Source/PythonScriptPlugin/Private/PyWrapperTypeRegistry.cpp
@@ -720,10 +720,10 @@ PyTypeObject* FPyWrapperTypeRegistry::GenerateWrappedClassType(const UClass* InC
 		}
 
 		const FString PythonStructMethodName = PyGenUtil::GetScriptMethodPythonName(InFunc);
-		TArray<PyGenUtil::FGeneratedWrappedDynamicMethod, TInlineAllocator<4>> DynamicMethodDefs;
+		TArray<TUniqueObj<PyGenUtil::FGeneratedWrappedDynamicMethod>, TInlineAllocator<4>> DynamicMethodDefs;
 
 		// Copy the basic wrapped method as we need to adjust some parts of it below
-		PyGenUtil::FGeneratedWrappedDynamicMethod& GeneratedWrappedDynamicMethod = DynamicMethodDefs.AddDefaulted_GetRef();
+		PyGenUtil::FGeneratedWrappedDynamicMethod& GeneratedWrappedDynamicMethod = DynamicMethodDefs.AddDefaulted_GetRef().Get();
 		static_cast<PyGenUtil::FGeneratedWrappedMethod&>(GeneratedWrappedDynamicMethod) = InTypeMethod;
 		GeneratedWrappedDynamicMethod.SelfParam = SelfParam;
 
@@ -814,7 +814,7 @@ PyTypeObject* FPyWrapperTypeRegistry::GenerateWrappedClassType(const UClass* InC
 		{
 			FString DeprecationMessage = FString::Printf(TEXT("'%s' was renamed to '%s'."), *DeprecatedPythonStructMethodName, *PythonStructMethodName);
 
-			PyGenUtil::FGeneratedWrappedDynamicMethod& DeprecatedGeneratedWrappedMethod = DynamicMethodDefs.AddDefaulted_GetRef();
+			PyGenUtil::FGeneratedWrappedDynamicMethod& DeprecatedGeneratedWrappedMethod = DynamicMethodDefs.AddDefaulted_GetRef().Get();
 			DeprecatedGeneratedWrappedMethod = GeneratedWrappedDynamicMethod;
 			DeprecatedGeneratedWrappedMethod.MethodName = PyGenUtil::TCHARToUTF8Buffer(*DeprecatedPythonStructMethodName);
 			DeprecatedGeneratedWrappedMethod.MethodDoc = PyGenUtil::TCHARToUTF8Buffer(*FString::Printf(TEXT("deprecated: %s"), *DeprecationMessage));
@@ -836,10 +836,10 @@ PyTypeObject* FPyWrapperTypeRegistry::GenerateWrappedClassType(const UClass* InC
 			check(HostedClassGeneratedWrappedType.IsValid());
 
 			// Add the dynamic methods to the class
-			for (PyGenUtil::FGeneratedWrappedDynamicMethod& GeneratedWrappedDynamicMethodToAdd : DynamicMethodDefs)
+			for (TUniqueObj<PyGenUtil::FGeneratedWrappedDynamicMethod>& GeneratedWrappedDynamicMethodToAdd : DynamicMethodDefs)
 			{
-				HostedClassGeneratedWrappedType->FieldTracker.RegisterPythonFieldName(UTF8_TO_TCHAR(GeneratedWrappedDynamicMethodToAdd.MethodName.GetData()), InFunc);
-				HostedClassGeneratedWrappedType->AddDynamicMethod(MoveTemp(GeneratedWrappedDynamicMethodToAdd));
+				HostedClassGeneratedWrappedType->FieldTracker.RegisterPythonFieldName(UTF8_TO_TCHAR(GeneratedWrappedDynamicMethodToAdd->MethodName.GetData()), InFunc);
+				HostedClassGeneratedWrappedType->AddDynamicMethod(MoveTemp(GeneratedWrappedDynamicMethodToAdd.Get()));
 			}
 		}
 		else if (SelfParam.ParamProp->IsA<UStructProperty>())
@@ -856,10 +856,10 @@ PyTypeObject* FPyWrapperTypeRegistry::GenerateWrappedClassType(const UClass* InC
 			check(HostedStructGeneratedWrappedType.IsValid());
 
 			// Add the dynamic methods to the struct
-			for (PyGenUtil::FGeneratedWrappedDynamicMethod& GeneratedWrappedDynamicMethodToAdd : DynamicMethodDefs)
+			for (TUniqueObj<PyGenUtil::FGeneratedWrappedDynamicMethod>& GeneratedWrappedDynamicMethodToAdd : DynamicMethodDefs)
 			{
-				HostedStructGeneratedWrappedType->FieldTracker.RegisterPythonFieldName(UTF8_TO_TCHAR(GeneratedWrappedDynamicMethodToAdd.MethodName.GetData()), InFunc);
-				HostedStructGeneratedWrappedType->AddDynamicMethod(MoveTemp(GeneratedWrappedDynamicMethodToAdd));
+				HostedStructGeneratedWrappedType->FieldTracker.RegisterPythonFieldName(UTF8_TO_TCHAR(GeneratedWrappedDynamicMethodToAdd->MethodName.GetData()), InFunc);
+				HostedStructGeneratedWrappedType->AddDynamicMethod(MoveTemp(GeneratedWrappedDynamicMethodToAdd.Get()));
 			}
 		}
 		else
@@ -960,10 +960,10 @@ PyTypeObject* FPyWrapperTypeRegistry::GenerateWrappedClassType(const UClass* InC
 		}
 
 		const FString PythonConstantName = PyGenUtil::GetScriptConstantPythonName(InFunc);
-		TArray<PyGenUtil::FGeneratedWrappedConstant, TInlineAllocator<4>> ConstantDefs;
+		TArray<TUniqueObj<PyGenUtil::FGeneratedWrappedConstant>, TInlineAllocator<4>> ConstantDefs;
 
 		// Build the constant definition
-		PyGenUtil::FGeneratedWrappedConstant& GeneratedWrappedConstant = ConstantDefs.AddDefaulted_GetRef();
+		PyGenUtil::FGeneratedWrappedConstant& GeneratedWrappedConstant = ConstantDefs.AddDefaulted_GetRef().Get();
 		GeneratedWrappedConstant.ConstantName = PyGenUtil::TCHARToUTF8Buffer(*PythonConstantName);
 		GeneratedWrappedConstant.ConstantDoc = PyGenUtil::TCHARToUTF8Buffer(*FString::Printf(TEXT("(%s): %s"), *PyGenUtil::GetPropertyPythonType(ConstantFunc.OutputParams[0].ParamProp), *PyGenUtil::GetFieldTooltip(InFunc)));
 		GeneratedWrappedConstant.ConstantFunc = ConstantFunc;
@@ -974,7 +974,7 @@ PyTypeObject* FPyWrapperTypeRegistry::GenerateWrappedClassType(const UClass* InC
 		{
 			FString DeprecationMessage = FString::Printf(TEXT("'%s' was renamed to '%s'."), *DeprecatedPythonConstantName, *PythonConstantName);
 
-			PyGenUtil::FGeneratedWrappedConstant& DeprecatedGeneratedWrappedConstant = ConstantDefs.AddDefaulted_GetRef();
+			PyGenUtil::FGeneratedWrappedConstant& DeprecatedGeneratedWrappedConstant = ConstantDefs.AddDefaulted_GetRef().Get();
 			DeprecatedGeneratedWrappedConstant = GeneratedWrappedConstant;
 			DeprecatedGeneratedWrappedConstant.ConstantName = PyGenUtil::TCHARToUTF8Buffer(*DeprecatedPythonConstantName);
 			DeprecatedGeneratedWrappedConstant.ConstantDoc = PyGenUtil::TCHARToUTF8Buffer(*FString::Printf(TEXT("deprecated: %s"), *DeprecationMessage));
@@ -998,10 +998,10 @@ PyTypeObject* FPyWrapperTypeRegistry::GenerateWrappedClassType(const UClass* InC
 				check(HostedClassGeneratedWrappedType.IsValid());
 
 				// Add the dynamic constants to the struct
-				for (PyGenUtil::FGeneratedWrappedConstant& GeneratedWrappedConstantToAdd : ConstantDefs)
+				for (TUniqueObj<PyGenUtil::FGeneratedWrappedConstant>& GeneratedWrappedConstantToAdd : ConstantDefs)
 				{
-					HostedClassGeneratedWrappedType->FieldTracker.RegisterPythonFieldName(UTF8_TO_TCHAR(GeneratedWrappedConstantToAdd.ConstantName.GetData()), InFunc);
-					HostedClassGeneratedWrappedType->AddDynamicConstant(MoveTemp(GeneratedWrappedConstantToAdd));
+					HostedClassGeneratedWrappedType->FieldTracker.RegisterPythonFieldName(UTF8_TO_TCHAR(GeneratedWrappedConstantToAdd->ConstantName.GetData()), InFunc);
+					HostedClassGeneratedWrappedType->AddDynamicConstant(MoveTemp(GeneratedWrappedConstantToAdd.Get()));
 				}
 			}
 			else if (HostType->IsA<UScriptStruct>())
@@ -1019,10 +1019,10 @@ PyTypeObject* FPyWrapperTypeRegistry::GenerateWrappedClassType(const UClass* InC
 				check(HostedStructGeneratedWrappedType.IsValid());
 
 				// Add the dynamic constants to the struct
-				for (PyGenUtil::FGeneratedWrappedConstant& GeneratedWrappedConstantToAdd : ConstantDefs)
+				for (TUniqueObj<PyGenUtil::FGeneratedWrappedConstant>& GeneratedWrappedConstantToAdd : ConstantDefs)
 				{
-					HostedStructGeneratedWrappedType->FieldTracker.RegisterPythonFieldName(UTF8_TO_TCHAR(GeneratedWrappedConstantToAdd.ConstantName.GetData()), InFunc);
-					HostedStructGeneratedWrappedType->AddDynamicConstant(MoveTemp(GeneratedWrappedConstantToAdd));
+					HostedStructGeneratedWrappedType->FieldTracker.RegisterPythonFieldName(UTF8_TO_TCHAR(GeneratedWrappedConstantToAdd->ConstantName.GetData()), InFunc);
+					HostedStructGeneratedWrappedType->AddDynamicConstant(MoveTemp(GeneratedWrappedConstantToAdd.Get()));
 				}
 			}
 			else
@@ -1033,10 +1033,10 @@ PyTypeObject* FPyWrapperTypeRegistry::GenerateWrappedClassType(const UClass* InC
 		else
 		{
 			// Add the static constants to this type
-			for (PyGenUtil::FGeneratedWrappedConstant& GeneratedWrappedConstantToAdd : ConstantDefs)
+			for (TUniqueObj<PyGenUtil::FGeneratedWrappedConstant>& GeneratedWrappedConstantToAdd : ConstantDefs)
 			{
-				GeneratedWrappedType->FieldTracker.RegisterPythonFieldName(UTF8_TO_TCHAR(GeneratedWrappedConstantToAdd.ConstantName.GetData()), InFunc);
-				GeneratedWrappedType->Constants.TypeConstants.Add(MoveTemp(GeneratedWrappedConstantToAdd));
+				GeneratedWrappedType->FieldTracker.RegisterPythonFieldName(UTF8_TO_TCHAR(GeneratedWrappedConstantToAdd->ConstantName.GetData()), InFunc);
+				GeneratedWrappedType->Constants.TypeConstants.Add(MoveTemp(GeneratedWrappedConstantToAdd.Get()));
 			}
 		}
 	};
