Generate UserDefinedEnum wrappers on-demand if marshalled via Python

This avoids the caller needing to have called unreal.get_blueprint_generated_types(...) for the type before trying to access it.

diff --git a/Engine/Plugins/Experimental/PythonScriptPlugin/Source/PythonScriptPlugin/Private/PyConversion.cpp b/Engine/Plugins/Experimental/PythonScriptPlugin/Source/PythonScriptPlugin/Private/PyConversion.cpp
--- a/Engine/Plugins/Experimental/PythonScriptPlugin/Source/PythonScriptPlugin/Private/PyConversion.cpp
+++ b/Engine/Plugins/Experimental/PythonScriptPlugin/Source/PythonScriptPlugin/Private/PyConversion.cpp
@@ -2,6 +2,7 @@
 
 #include "PyConversion.h"
 #include "PyUtil.h"
+#include "PyGenUtil.h"
 #include "PyPtr.h"
 
 #include "PyWrapperObject.h"
@@ -628,10 +629,28 @@ PyObject* PythonizeStruct(UScriptStruct* Val, const ESetErrorState SetErrorState
 	return Obj;
 }
 
+void EnsureWrappedBlueprintEnumType(const UEnum* EnumType)
+{
+	if (PyGenUtil::IsBlueprintGeneratedEnum(EnumType))
+	{
+		FPyWrapperTypeRegistry& PyWrapperTypeRegistry = FPyWrapperTypeRegistry::Get();
+
+		FPyWrapperTypeRegistry::FGeneratedWrappedTypeReferences GeneratedWrappedTypeReferences;
+		TSet<FName> DirtyModules;
+
+		PyWrapperTypeRegistry.GenerateWrappedTypeForObject(EnumType, GeneratedWrappedTypeReferences, DirtyModules, EPyTypeGenerationFlags::IncludeBlueprintGeneratedTypes);
+
+		PyWrapperTypeRegistry.GenerateWrappedTypesForReferences(GeneratedWrappedTypeReferences, DirtyModules);
+		PyWrapperTypeRegistry.NotifyModulesDirtied(DirtyModules);
+	}
+}
+
 FPyConversionResult NativizeEnumEntry(PyObject* PyObj, const UEnum* EnumType, int64& OutVal, const ESetErrorState SetErrorState)
 {
 	FPyConversionResult Result = FPyConversionResult::Failure();
 
+	EnsureWrappedBlueprintEnumType(EnumType);
+
 	PyTypeObject* PyEnumType = FPyWrapperTypeRegistry::Get().GetWrappedEnumType(EnumType);
 	FPyWrapperEnumPtr PyEnum = FPyWrapperEnumPtr::StealReference(FPyWrapperEnum::CastPyObject(PyObj, PyEnumType, &Result));
 	if (PyEnum)
@@ -644,6 +663,8 @@ FPyConversionResult NativizeEnumEntry(PyObject* PyObj, const UEnum* EnumType, in
 
 FPyConversionResult PythonizeEnumEntry(const int64 Val, const UEnum* EnumType, PyObject*& OutPyObj, const ESetErrorState SetErrorState)
 {
+	EnsureWrappedBlueprintEnumType(EnumType);
+
 	PyTypeObject* PyEnumType = FPyWrapperTypeRegistry::Get().GetWrappedEnumType(EnumType);
 	if (const FPyWrapperEnumMetaData* PyEnumMetaData = FPyWrapperEnumMetaData::GetMetaData(PyEnumType))
 	{
