Added a way to force a change notify when using SetEditorProperty in Python or Editor Utilities, even if the value is unchanged

This may be needed if you've edited data on an instance indirectly (eg, modifying an array reference on the object), and need to force a change notification to update some dependent data.

In Python, this can be done by passing unreal.PropertyAccessChangeNotifyMode.ALWAYS as the notify_mode argument to set_editor_property.
In Editor Utilities, there is a new advanced pin available for setting the "Change Notify Mode".
In both cases the default method is to notify only when the value actually changes.

This change also significantly simplifies UK2Node_GetEditorProperty and UK2Node_SetEditorProperty by having them derive from UK2Node_CallFunction, rather than being a proxy node that expands to a UK2Node_CallFunction node.

diff --git a/Engine/Plugins/Experimental/PythonScriptPlugin/Source/PythonScriptPlugin/Private/PyConversion.cpp b/Engine/Plugins/Experimental/PythonScriptPlugin/Source/PythonScriptPlugin/Private/PyConversion.cpp
--- a/Engine/Plugins/Experimental/PythonScriptPlugin/Source/PythonScriptPlugin/Private/PyConversion.cpp
+++ b/Engine/Plugins/Experimental/PythonScriptPlugin/Source/PythonScriptPlugin/Private/PyConversion.cpp
@@ -701,7 +701,7 @@ FPyConversionResult NativizeProperty(PyObject* PyObj, const UProperty* Prop, voi
 		FPyWrapperFixedArrayPtr PyFixedArray = FPyWrapperFixedArrayPtr::StealReference(FPyWrapperFixedArray::CastPyObject(PyObj, &PyWrapperFixedArrayType, Prop));
 		if (PyFixedArray)
 		{
-			EmitPropertyChangeNotifications(InChangeNotify, [&]()
+			EmitPropertyChangeNotifications(InChangeNotify, /*bIdenticalValue*/false, [&]()
 			{
 				const int32 ArrSize = FMath::Min(Prop->ArrayDim, PyFixedArray->ArrayProp->ArrayDim);
 				for (int32 ArrIndex = 0; ArrIndex < ArrSize; ++ArrIndex)
@@ -744,13 +744,11 @@ FPyConversionResult NativizeProperty_Direct(PyObject* PyObj, const UProperty* Pr
 		if (Result)																	\
 		{																			\
 			auto OldValue = CastProp->GetPropertyValue(ValueAddr);					\
-			if (OldValue != NewValue)												\
+			EmitPropertyChangeNotifications(InChangeNotify,							\
+				OldValue == NewValue, [&]()											\
 			{																		\
-				EmitPropertyChangeNotifications(InChangeNotify, [&]()				\
-				{																	\
-					CastProp->SetPropertyValue(ValueAddr, NewValue);				\
-				});																	\
-			}																		\
+				CastProp->SetPropertyValue(ValueAddr, NewValue);					\
+			});																		\
 		}																			\
 		PYCONVERSION_PROPERTY_RETURN(Result);										\
 	}
@@ -763,13 +761,11 @@ FPyConversionResult NativizeProperty_Direct(PyObject* PyObj, const UProperty* Pr
 		if (Result)																	\
 		{																			\
 			auto* ValuePtr = static_cast<PROPTYPE::TCppType*>(ValueAddr);			\
-			if (!CastProp->Identical(ValuePtr, &NewValue, PPF_None))				\
+			EmitPropertyChangeNotifications(InChangeNotify,							\
+				CastProp->Identical(ValuePtr, &NewValue, PPF_None), [&]()			\
 			{																		\
-				EmitPropertyChangeNotifications(InChangeNotify, [&]()				\
-				{																	\
-					*ValuePtr = MoveTemp(NewValue);									\
-				});																	\
-			}																		\
+				*ValuePtr = MoveTemp(NewValue);										\
+			});																		\
 		}																			\
 		PYCONVERSION_PROPERTY_RETURN(Result);										\
 	}
@@ -815,13 +811,10 @@ FPyConversionResult NativizeProperty_Direct(PyObject* PyObj, const UProperty* Pr
 		if (Result)
 		{
 			auto* ValuePtr = static_cast<uint8*>(ValueAddr);
-			if (*ValuePtr != NewValue)
+			EmitPropertyChangeNotifications(InChangeNotify, *ValuePtr == NewValue, [&]()
 			{
-				EmitPropertyChangeNotifications(InChangeNotify, [&]()
-				{
-					*ValuePtr = NewValue;
-				});
-			}
+				*ValuePtr = NewValue;
+			});
 		}
 		PYCONVERSION_PROPERTY_RETURN(Result);
 	}
@@ -845,13 +838,10 @@ FPyConversionResult NativizeProperty_Direct(PyObject* PyObj, const UProperty* Pr
 			if (Result)
 			{
 				const int64 OldValue = EnumInternalProp->GetSignedIntPropertyValue(ValueAddr);
-				if (OldValue != NewValue)
+				EmitPropertyChangeNotifications(InChangeNotify, OldValue == NewValue, [&]()
 				{
-					EmitPropertyChangeNotifications(InChangeNotify, [&]()
-					{
-						EnumInternalProp->SetIntPropertyValue(ValueAddr, NewValue);
-					});
-				}
+					EnumInternalProp->SetIntPropertyValue(ValueAddr, NewValue);
+				});
 			}
 		}
 
@@ -865,13 +855,10 @@ FPyConversionResult NativizeProperty_Direct(PyObject* PyObj, const UProperty* Pr
 		if (Result)
 		{
 			UObject* OldValue = CastProp->GetObjectPropertyValue(ValueAddr);
-			if (OldValue != NewValue)
+			EmitPropertyChangeNotifications(InChangeNotify, OldValue == NewValue, [&]()
 			{
-				EmitPropertyChangeNotifications(InChangeNotify, [&]()
-				{
-					CastProp->SetObjectPropertyValue(ValueAddr, NewValue);
-				});
-			}
+				CastProp->SetObjectPropertyValue(ValueAddr, NewValue);
+			});
 		}
 		PYCONVERSION_PROPERTY_RETURN(Result);
 	}
@@ -883,13 +870,10 @@ FPyConversionResult NativizeProperty_Direct(PyObject* PyObj, const UProperty* Pr
 		if (Result)
 		{
 			UObject* OldValue = CastProp->GetObjectPropertyValue(ValueAddr);
-			if (OldValue != NewValue)
+			EmitPropertyChangeNotifications(InChangeNotify, OldValue == NewValue, [&]()
 			{
-				EmitPropertyChangeNotifications(InChangeNotify, [&]()
-				{
-					CastProp->SetObjectPropertyValue(ValueAddr, NewValue);
-				});
-			}
+				CastProp->SetObjectPropertyValue(ValueAddr, NewValue);
+			});
 		}
 		PYCONVERSION_PROPERTY_RETURN(Result);
 	}
@@ -901,13 +885,10 @@ FPyConversionResult NativizeProperty_Direct(PyObject* PyObj, const UProperty* Pr
 		if (Result)
 		{
 			UObject* OldValue = CastProp->GetObjectPropertyValue(ValueAddr);
-			if (OldValue != NewValue)
+			EmitPropertyChangeNotifications(InChangeNotify, OldValue == NewValue, [&]()
 			{
-				EmitPropertyChangeNotifications(InChangeNotify, [&]()
-				{
-					CastProp->SetObjectPropertyValue(ValueAddr, NewValue);
-				});
-			}
+				CastProp->SetObjectPropertyValue(ValueAddr, NewValue);
+			});
 		}
 		PYCONVERSION_PROPERTY_RETURN(Result);
 	}
@@ -919,13 +900,10 @@ FPyConversionResult NativizeProperty_Direct(PyObject* PyObj, const UProperty* Pr
 		if (Result)
 		{
 			UObject* OldValue = CastProp->GetPropertyValue(ValueAddr).GetObject();
-			if (OldValue != NewValue)
+			EmitPropertyChangeNotifications(InChangeNotify, OldValue == NewValue, [&]()
 			{
-				EmitPropertyChangeNotifications(InChangeNotify, [&]()
-				{
-					CastProp->SetPropertyValue(ValueAddr, FScriptInterface(NewValue, NewValue ? NewValue->GetInterfaceAddress(CastProp->InterfaceClass) : nullptr));
-				});
-			}
+				CastProp->SetPropertyValue(ValueAddr, FScriptInterface(NewValue, NewValue ? NewValue->GetInterfaceAddress(CastProp->InterfaceClass) : nullptr));
+			});
 		}
 		PYCONVERSION_PROPERTY_RETURN(Result);
 	}
@@ -937,13 +915,10 @@ FPyConversionResult NativizeProperty_Direct(PyObject* PyObj, const UProperty* Pr
 		FPyWrapperStructPtr PyStruct = FPyWrapperStructPtr::StealReference(FPyWrapperStruct::CastPyObject(PyObj, PyStructType, &Result));
 		if (PyStruct && ensureAlways(PyStruct->ScriptStruct->IsChildOf(CastProp->Struct)))
 		{
-			if (!CastProp->Identical(ValueAddr, PyStruct->StructInstance, PPF_None))
+			EmitPropertyChangeNotifications(InChangeNotify, CastProp->Identical(ValueAddr, PyStruct->StructInstance, PPF_None), [&]()
 			{
-				EmitPropertyChangeNotifications(InChangeNotify, [&]()
-				{
-					CastProp->Struct->CopyScriptStruct(ValueAddr, PyStruct->StructInstance);
-				});
-			}
+				CastProp->Struct->CopyScriptStruct(ValueAddr, PyStruct->StructInstance);
+			});
 		}
 		else
 		{
@@ -963,13 +938,10 @@ FPyConversionResult NativizeProperty_Direct(PyObject* PyObj, const UProperty* Pr
 		FPyWrapperDelegatePtr PyDelegate = FPyWrapperDelegatePtr::StealReference(FPyWrapperDelegate::CastPyObject(PyObj, PyDelegateType, &Result));
 		if (PyDelegate)
 		{
-			if (!CastProp->Identical(ValueAddr, PyDelegate->DelegateInstance, PPF_None))
+			EmitPropertyChangeNotifications(InChangeNotify, CastProp->Identical(ValueAddr, PyDelegate->DelegateInstance, PPF_None), [&]()
 			{
-				EmitPropertyChangeNotifications(InChangeNotify, [&]()
-				{
-					CastProp->SetPropertyValue(ValueAddr, *PyDelegate->DelegateInstance);
-				});
-			}
+				CastProp->SetPropertyValue(ValueAddr, *PyDelegate->DelegateInstance);
+			});
 		}
 		PYCONVERSION_PROPERTY_RETURN(Result);
 	}
@@ -981,13 +953,10 @@ FPyConversionResult NativizeProperty_Direct(PyObject* PyObj, const UProperty* Pr
 		FPyWrapperMulticastDelegatePtr PyDelegate = FPyWrapperMulticastDelegatePtr::StealReference(FPyWrapperMulticastDelegate::CastPyObject(PyObj, PyDelegateType, &Result));
 		if (PyDelegate)
 		{
-			if (!CastProp->Identical(ValueAddr, PyDelegate->DelegateInstance, PPF_None))
+			EmitPropertyChangeNotifications(InChangeNotify, CastProp->Identical(ValueAddr, PyDelegate->DelegateInstance, PPF_None), [&]()
 			{
-				EmitPropertyChangeNotifications(InChangeNotify, [&]()
-				{
-					CastProp->SetMulticastDelegate(ValueAddr, *PyDelegate->DelegateInstance);
-				});
-			}
+				CastProp->SetMulticastDelegate(ValueAddr, *PyDelegate->DelegateInstance);
+			});
 		}
 		PYCONVERSION_PROPERTY_RETURN(Result);
 	}
@@ -998,13 +967,10 @@ FPyConversionResult NativizeProperty_Direct(PyObject* PyObj, const UProperty* Pr
 		FPyWrapperArrayPtr PyArray = FPyWrapperArrayPtr::StealReference(FPyWrapperArray::CastPyObject(PyObj, &PyWrapperArrayType, CastProp->Inner, &Result));
 		if (PyArray)
 		{
-			if (!CastProp->Identical(ValueAddr, PyArray->ArrayInstance, PPF_None))
+			EmitPropertyChangeNotifications(InChangeNotify, CastProp->Identical(ValueAddr, PyArray->ArrayInstance, PPF_None), [&]()
 			{
-				EmitPropertyChangeNotifications(InChangeNotify, [&]()
-				{
-					CastProp->CopySingleValue(ValueAddr, PyArray->ArrayInstance);
-				});
-			}
+				CastProp->CopySingleValue(ValueAddr, PyArray->ArrayInstance);
+			});
 		}
 		PYCONVERSION_PROPERTY_RETURN(Result);
 	}
@@ -1015,13 +981,10 @@ FPyConversionResult NativizeProperty_Direct(PyObject* PyObj, const UProperty* Pr
 		FPyWrapperSetPtr PySet = FPyWrapperSetPtr::StealReference(FPyWrapperSet::CastPyObject(PyObj, &PyWrapperSetType, CastProp->ElementProp, &Result));
 		if (PySet)
 		{
-			if (!CastProp->Identical(ValueAddr, PySet->SetInstance, PPF_None))
+			EmitPropertyChangeNotifications(InChangeNotify, CastProp->Identical(ValueAddr, PySet->SetInstance, PPF_None), [&]()
 			{
-				EmitPropertyChangeNotifications(InChangeNotify, [&]()
-				{
-					CastProp->CopySingleValue(ValueAddr, PySet->SetInstance);
-				});
-			}
+				CastProp->CopySingleValue(ValueAddr, PySet->SetInstance);
+			});
 		}
 		PYCONVERSION_PROPERTY_RETURN(Result);
 	}
@@ -1032,13 +995,10 @@ FPyConversionResult NativizeProperty_Direct(PyObject* PyObj, const UProperty* Pr
 		FPyWrapperMapPtr PyMap = FPyWrapperMapPtr::StealReference(FPyWrapperMap::CastPyObject(PyObj, &PyWrapperMapType, CastProp->KeyProp, CastProp->ValueProp, &Result));
 		if (PyMap)
 		{
-			if (!CastProp->Identical(ValueAddr, PyMap->MapInstance, PPF_None))
+			EmitPropertyChangeNotifications(InChangeNotify, CastProp->Identical(ValueAddr, PyMap->MapInstance, PPF_None), [&]()
 			{
-				EmitPropertyChangeNotifications(InChangeNotify, [&]()
-				{
-					CastProp->CopySingleValue(ValueAddr, PyMap->MapInstance);
-				});
-			}
+				CastProp->CopySingleValue(ValueAddr, PyMap->MapInstance);
+			});
 		}
 		PYCONVERSION_PROPERTY_RETURN(Result);
 	}
@@ -1197,11 +1157,14 @@ FPyConversionResult PythonizeProperty_InContainer(const UProperty* Prop, const v
 	return PythonizeProperty(Prop, Prop->ContainerPtrToValuePtr<void>(BaseAddr, ArrayIndex), OutPyObj, ConversionMethod, OwnerPyObj, SetErrorState);
 }
 
-void EmitPropertyChangeNotifications(const FPropertyAccessChangeNotify* InChangeNotify, const TFunctionRef<void()>& InDoChangeFunc)
+void EmitPropertyChangeNotifications(const FPropertyAccessChangeNotify* InChangeNotify, const bool bIdenticalValue, const TFunctionRef<void()>& InDoChangeFunc)
 {
-	PropertyAccessUtil::EmitPreChangeNotify(InChangeNotify);
-	InDoChangeFunc();
-	PropertyAccessUtil::EmitPostChangeNotify(InChangeNotify);
+	PropertyAccessUtil::EmitPreChangeNotify(InChangeNotify, bIdenticalValue);
+	if (!bIdenticalValue)
+	{
+		InDoChangeFunc();
+	}
+	PropertyAccessUtil::EmitPostChangeNotify(InChangeNotify, bIdenticalValue);
 }
 
 }
diff --git a/Engine/Plugins/Experimental/PythonScriptPlugin/Source/PythonScriptPlugin/Private/PyConversion.h b/Engine/Plugins/Experimental/PythonScriptPlugin/Source/PythonScriptPlugin/Private/PyConversion.h
--- a/Engine/Plugins/Experimental/PythonScriptPlugin/Source/PythonScriptPlugin/Private/PyConversion.h
+++ b/Engine/Plugins/Experimental/PythonScriptPlugin/Source/PythonScriptPlugin/Private/PyConversion.h
@@ -234,7 +234,7 @@ namespace PyConversion
 	 * This function should be called when you know the value will actually change (or know you want to emit the notifications for it changing) and will do 
 	 * the pre-change notify, invoke the passed delegate to perform the change, then do the post-change notify
 	 */
-	void EmitPropertyChangeNotifications(const FPropertyAccessChangeNotify* InChangeNotify, const TFunctionRef<void()>& InDoChangeFunc);
+	void EmitPropertyChangeNotifications(const FPropertyAccessChangeNotify* InChangeNotify, const bool bIdenticalValue, const TFunctionRef<void()>& InDoChangeFunc);
 }
 
 #endif	// WITH_PYTHON
diff --git a/Engine/Plugins/Experimental/PythonScriptPlugin/Source/PythonScriptPlugin/Private/PyWrapperObject.cpp b/Engine/Plugins/Experimental/PythonScriptPlugin/Source/PythonScriptPlugin/Private/PyWrapperObject.cpp
--- a/Engine/Plugins/Experimental/PythonScriptPlugin/Source/PythonScriptPlugin/Private/PyWrapperObject.cpp
+++ b/Engine/Plugins/Experimental/PythonScriptPlugin/Source/PythonScriptPlugin/Private/PyWrapperObject.cpp
@@ -126,14 +126,14 @@ PyObject* FPyWrapperObject::GetPropertyValue(FPyWrapperObject* InSelf, const PyG
 	return PyGenUtil::GetPropertyValue(InSelf->ObjectInstance->GetClass(), InSelf->ObjectInstance, InPropDef, InPythonAttrName, (PyObject*)InSelf, *PyUtil::GetErrorContext(InSelf));
 }
 
-int FPyWrapperObject::SetPropertyValue(FPyWrapperObject* InSelf, PyObject* InValue, const PyGenUtil::FGeneratedWrappedProperty& InPropDef, const char* InPythonAttrName, const bool InNotifyChange, const uint64 InReadOnlyFlags)
+int FPyWrapperObject::SetPropertyValue(FPyWrapperObject* InSelf, PyObject* InValue, const PyGenUtil::FGeneratedWrappedProperty& InPropDef, const char* InPythonAttrName, const EPropertyAccessChangeNotifyMode InNotifyMode, const uint64 InReadOnlyFlags)
 {
 	if (!ValidateInternalState(InSelf))
 	{
 		return -1;
 	}
 
-	const TUniquePtr<FPropertyAccessChangeNotify> ChangeNotify = InNotifyChange ? FPyWrapperOwnerContext((PyObject*)InSelf, InPropDef.Prop).BuildChangeNotify() : nullptr;
+	const TUniquePtr<FPropertyAccessChangeNotify> ChangeNotify = FPyWrapperOwnerContext((PyObject*)InSelf, InPropDef.Prop).BuildChangeNotify(InNotifyMode);
 	return PyGenUtil::SetPropertyValue(InSelf->ObjectInstance->GetClass(), InSelf->ObjectInstance, InValue, InPropDef, InPythonAttrName, ChangeNotify.Get(), InReadOnlyFlags, PropertyAccessUtil::IsObjectTemplate(InSelf->ObjectInstance), *PyUtil::GetErrorContext(InSelf));
 }
 
@@ -823,9 +823,10 @@ PyTypeObject InitializePyWrapperObjectType()
 
 			PyObject* PyNameObj = nullptr;
 			PyObject* PyValueObj = nullptr;
+			PyObject* PyNotifyModeObj = nullptr;
 
-			static const char *ArgsKwdList[] = { "name", "value", nullptr };
-			if (!PyArg_ParseTupleAndKeywords(InArgs, InKwds, "OO:set_editor_property", (char**)ArgsKwdList, &PyNameObj, &PyValueObj))
+			static const char *ArgsKwdList[] = { "name", "value", "notify_mode", nullptr };
+			if (!PyArg_ParseTupleAndKeywords(InArgs, InKwds, "OO|O:set_editor_property", (char**)ArgsKwdList, &PyNameObj, &PyValueObj, &PyNotifyModeObj))
 			{
 				return nullptr;
 			}
@@ -837,13 +838,22 @@ PyTypeObject InitializePyWrapperObjectType()
 				return nullptr;
 			}
 
+			static const UEnum* PropertyAccessChangeNotifyModeEnum = FindObjectChecked<UEnum>(ANY_PACKAGE, TEXT("EPropertyAccessChangeNotifyMode"));
+
+			EPropertyAccessChangeNotifyMode NotifyMode = EPropertyAccessChangeNotifyMode::Default;
+			if (PyNotifyModeObj && !PyConversion::NativizeEnumEntry(PyNotifyModeObj, PropertyAccessChangeNotifyModeEnum, NotifyMode))
+			{
+				PyUtil::SetPythonError(PyExc_TypeError, InSelf, *FString::Printf(TEXT("Failed to convert 'notify_mode' (%s) to 'PropertyAccessChangeNotifyMode'"), *PyUtil::GetFriendlyTypename(PyNotifyModeObj)));
+				return nullptr;
+			}
+
 			const PyGenUtil::FGeneratedWrappedProperty WrappedPropDef = FindEditorPropertyImpl(InSelf, Name);
 			if (!WrappedPropDef.Prop)
 			{
 				return nullptr;
 			}
 
-			const int Result = FPyWrapperObject::SetPropertyValue(InSelf, PyValueObj, WrappedPropDef, TCHAR_TO_UTF8(*Name.ToString()), /*InNotifyChange*/true, PropertyAccessUtil::EditorReadOnlyFlags);
+			const int Result = FPyWrapperObject::SetPropertyValue(InSelf, PyValueObj, WrappedPropDef, TCHAR_TO_UTF8(*Name.ToString()), NotifyMode, PropertyAccessUtil::EditorReadOnlyFlags);
 			if (Result != 0)
 			{
 				return nullptr;
@@ -924,7 +934,7 @@ PyTypeObject InitializePyWrapperObjectType()
 				const PyGenUtil::FGeneratedWrappedProperty& WrappedPropDef = PropertyInfo.Get<1>();
 				PyObject* PyValueObj = PropertyInfo.Get<2>().GetPtr();
 
-				const int Result = FPyWrapperObject::SetPropertyValue(InSelf, PyValueObj, WrappedPropDef, TCHAR_TO_UTF8(*Name.ToString()), /*InNotifyChange*/false, PropertyAccessUtil::EditorReadOnlyFlags);
+				const int Result = FPyWrapperObject::SetPropertyValue(InSelf, PyValueObj, WrappedPropDef, TCHAR_TO_UTF8(*Name.ToString()), EPropertyAccessChangeNotifyMode::Never, PropertyAccessUtil::EditorReadOnlyFlags);
 				if (Result != 0)
 				{
 					return nullptr;
@@ -952,7 +962,7 @@ PyTypeObject InitializePyWrapperObjectType()
 		{ "modify", PyCFunctionCast(&FMethods::Modify), METH_VARARGS, "x.modify(bool) -> bool -- inform that this instance is about to be modified (tracks changes for undo/redo if transactional)" },
 		{ "rename", PyCFunctionCast(&FMethods::Rename), METH_VARARGS | METH_KEYWORDS, "x.rename(name=None, outer=None) -> bool -- rename this instance" },
 		{ "get_editor_property", PyCFunctionCast(&FMethods::GetEditorProperty), METH_VARARGS | METH_KEYWORDS, "x.get_editor_property(name) -> object -- get the value of any property visible to the editor" },
-		{ "set_editor_property", PyCFunctionCast(&FMethods::SetEditorProperty), METH_VARARGS | METH_KEYWORDS, "x.set_editor_property(name, value) -> None -- set the value of any property visible to the editor, ensuring that the pre/post change notifications are called" },
+		{ "set_editor_property", PyCFunctionCast(&FMethods::SetEditorProperty), METH_VARARGS | METH_KEYWORDS, "x.set_editor_property(name, value, notify_mode=PropertyAccessChangeNotifyMode.DEFAULT) -> None -- set the value of any property visible to the editor, ensuring that the pre/post change notifications are called" },
 		{ "set_editor_properties", PyCFunctionCast(&FMethods::SetEditorProperties), METH_VARARGS, "x.set_editor_properties(property_info) -> None -- set the value of any properties visible to the editor (from a name->value dict), ensuring that the pre/post change notifications are called" },
 		{ nullptr, nullptr, 0, nullptr }
 	};
diff --git a/Engine/Plugins/Experimental/PythonScriptPlugin/Source/PythonScriptPlugin/Private/PyWrapperObject.h b/Engine/Plugins/Experimental/PythonScriptPlugin/Source/PythonScriptPlugin/Private/PyWrapperObject.h
--- a/Engine/Plugins/Experimental/PythonScriptPlugin/Source/PythonScriptPlugin/Private/PyWrapperObject.h
+++ b/Engine/Plugins/Experimental/PythonScriptPlugin/Source/PythonScriptPlugin/Private/PyWrapperObject.h
@@ -47,7 +47,7 @@ struct FPyWrapperObject : public FPyWrapperBase
 	static PyObject* GetPropertyValue(FPyWrapperObject* InSelf, const PyGenUtil::FGeneratedWrappedProperty& InPropDef, const char* InPythonAttrName);
 
 	/** Set a property value on this instance (called via generated code) */
-	static int SetPropertyValue(FPyWrapperObject* InSelf, PyObject* InValue, const PyGenUtil::FGeneratedWrappedProperty& InPropDef, const char* InPythonAttrName, const bool InNotifyChange = false, const uint64 InReadOnlyFlags = PropertyAccessUtil::RuntimeReadOnlyFlags);
+	static int SetPropertyValue(FPyWrapperObject* InSelf, PyObject* InValue, const PyGenUtil::FGeneratedWrappedProperty& InPropDef, const char* InPythonAttrName, const EPropertyAccessChangeNotifyMode InNotifyMode = EPropertyAccessChangeNotifyMode::Never, const uint64 InReadOnlyFlags = PropertyAccessUtil::RuntimeReadOnlyFlags);
 
 	/** Call a named getter function on this class using the given instance (called via generated code) */
 	static PyObject* CallGetterFunction(FPyWrapperObject* InSelf, const PyGenUtil::FGeneratedWrappedFunction& InFuncDef);
diff --git a/Engine/Plugins/Experimental/PythonScriptPlugin/Source/PythonScriptPlugin/Private/PyWrapperOwnerContext.cpp b/Engine/Plugins/Experimental/PythonScriptPlugin/Source/PythonScriptPlugin/Private/PyWrapperOwnerContext.cpp
--- a/Engine/Plugins/Experimental/PythonScriptPlugin/Source/PythonScriptPlugin/Private/PyWrapperOwnerContext.cpp
+++ b/Engine/Plugins/Experimental/PythonScriptPlugin/Source/PythonScriptPlugin/Private/PyWrapperOwnerContext.cpp
@@ -52,67 +52,70 @@ void FPyWrapperOwnerContext::AssertValidConversionMethod(const EPyConversionMeth
 	::AssertValidPyConversionOwner(GetOwnerObject(), InMethod);
 }
 
-TUniquePtr<FPropertyAccessChangeNotify> FPyWrapperOwnerContext::BuildChangeNotify() const
+TUniquePtr<FPropertyAccessChangeNotify> FPyWrapperOwnerContext::BuildChangeNotify(const EPropertyAccessChangeNotifyMode InNotifyMode) const
 {
 #if WITH_EDITOR
-	TUniquePtr<FPropertyAccessChangeNotify> ChangeNotify = MakeUnique<FPropertyAccessChangeNotify>();
-
-	auto AppendOwnerPropertyToChain = [&ChangeNotify](const FPyWrapperOwnerContext& InOwnerContext) -> bool
+	if (InNotifyMode != EPropertyAccessChangeNotifyMode::Never)
 	{
-		const UProperty* LeafProp = nullptr;
-		if (PyObject_IsInstance(InOwnerContext.GetOwnerObject(), (PyObject*)&PyWrapperObjectType) == 1 || PyObject_IsInstance(InOwnerContext.GetOwnerObject(), (PyObject*)&PyWrapperStructType) == 1)
-		{
-			LeafProp = InOwnerContext.GetOwnerProperty();
-		}
+		TUniquePtr<FPropertyAccessChangeNotify> ChangeNotify = MakeUnique<FPropertyAccessChangeNotify>();
+		ChangeNotify->NotifyMode = InNotifyMode;
 
-		if (LeafProp)
+		auto AppendOwnerPropertyToChain = [&ChangeNotify](const FPyWrapperOwnerContext& InOwnerContext) -> bool
 		{
-			ChangeNotify->ChangedPropertyChain.AddHead(const_cast<UProperty*>(LeafProp));
-			return true;
-		}
-
-		return false;
-	};
-
-	FPyWrapperOwnerContext OwnerContext = *this;
-	while (OwnerContext.HasOwner() && AppendOwnerPropertyToChain(OwnerContext))
-	{
-		PyObject* PyObj = OwnerContext.GetOwnerObject();
-
-		if (PyObj == GetOwnerObject())
+			const UProperty* LeafProp = nullptr;
+			if (PyObject_IsInstance(InOwnerContext.GetOwnerObject(), (PyObject*)&PyWrapperObjectType) == 1 || PyObject_IsInstance(InOwnerContext.GetOwnerObject(), (PyObject*)&PyWrapperStructType) == 1)
+			{
+				LeafProp = InOwnerContext.GetOwnerProperty();
+			}
+
+			if (LeafProp)
+			{
+				ChangeNotify->ChangedPropertyChain.AddHead(const_cast<UProperty*>(LeafProp));
+				return true;
+			}
+
+			return false;
+		};
+
+		FPyWrapperOwnerContext OwnerContext = *this;
+		while (OwnerContext.HasOwner() && AppendOwnerPropertyToChain(OwnerContext))
 		{
-			ChangeNotify->ChangedPropertyChain.SetActivePropertyNode(ChangeNotify->ChangedPropertyChain.GetHead()->GetValue());
-		}
-
-		if (PyObject_IsInstance(PyObj, (PyObject*)&PyWrapperObjectType) == 1)
-		{
-			// Found an object, this is the end of the chain
-			ChangeNotify->ChangedObject = ((FPyWrapperObject*)PyObj)->ObjectInstance;
-			ChangeNotify->ChangedPropertyChain.SetActiveMemberPropertyNode(ChangeNotify->ChangedPropertyChain.GetHead()->GetValue());
+			PyObject* PyObj = OwnerContext.GetOwnerObject();
+
+			if (PyObj == GetOwnerObject())
+			{
+				ChangeNotify->ChangedPropertyChain.SetActivePropertyNode(ChangeNotify->ChangedPropertyChain.GetHead()->GetValue());
+			}
+
+			if (PyObject_IsInstance(PyObj, (PyObject*)&PyWrapperObjectType) == 1)
+			{
+				// Found an object, this is the end of the chain
+				ChangeNotify->ChangedObject = ((FPyWrapperObject*)PyObj)->ObjectInstance;
+				ChangeNotify->ChangedPropertyChain.SetActiveMemberPropertyNode(ChangeNotify->ChangedPropertyChain.GetHead()->GetValue());
+				break;
+			}
+
+			if (PyObject_IsInstance(PyObj, (PyObject*)&PyWrapperStructType) == 1)
+			{
+				// Found a struct, recurse up the chain
+				OwnerContext = ((FPyWrapperStruct*)PyObj)->OwnerContext;
+				continue;
+			}
+
+			// Unknown object type - just bail
 			break;
 		}
 
-		if (PyObject_IsInstance(PyObj, (PyObject*)&PyWrapperStructType) == 1)
+		// If we didn't find an object in the chain then we can't emit notifications
+		if (!ChangeNotify->ChangedObject)
 		{
-			// Found a struct, recurse up the chain
-			OwnerContext = ((FPyWrapperStruct*)PyObj)->OwnerContext;
-			continue;
+			ChangeNotify.Reset();
 		}
 
-		// Unknown object type - just bail
-		break;
-	}
-
-	// If we didn't find an object in the chain then we can't emit notifications
-	if (!ChangeNotify->ChangedObject)
-	{
-		ChangeNotify.Reset();
+		return ChangeNotify;
 	}
-
-	return ChangeNotify;
-#else
-	return nullptr;
 #endif
+	return nullptr;
 }
 
 #endif	// WITH_PYTHON
diff --git a/Engine/Plugins/Experimental/PythonScriptPlugin/Source/PythonScriptPlugin/Private/PyWrapperOwnerContext.h b/Engine/Plugins/Experimental/PythonScriptPlugin/Source/PythonScriptPlugin/Private/PyWrapperOwnerContext.h
--- a/Engine/Plugins/Experimental/PythonScriptPlugin/Source/PythonScriptPlugin/Private/PyWrapperOwnerContext.h
+++ b/Engine/Plugins/Experimental/PythonScriptPlugin/Source/PythonScriptPlugin/Private/PyWrapperOwnerContext.h
@@ -38,7 +38,7 @@ public:
 	void AssertValidConversionMethod(const EPyConversionMethod InMethod) const;
 
 	/** Build the property change notify that corresponds to this owner context, or null if this owner context shouldn't emit change notifications */
-	TUniquePtr<FPropertyAccessChangeNotify> BuildChangeNotify() const;
+	TUniquePtr<FPropertyAccessChangeNotify> BuildChangeNotify(const EPropertyAccessChangeNotifyMode InNotifyMode) const;
 
 private:
 	/** The Python object that owns the instance being wrapped (if any) */
diff --git a/Engine/Plugins/Experimental/PythonScriptPlugin/Source/PythonScriptPlugin/Private/PyWrapperStruct.cpp b/Engine/Plugins/Experimental/PythonScriptPlugin/Source/PythonScriptPlugin/Private/PyWrapperStruct.cpp
--- a/Engine/Plugins/Experimental/PythonScriptPlugin/Source/PythonScriptPlugin/Private/PyWrapperStruct.cpp
+++ b/Engine/Plugins/Experimental/PythonScriptPlugin/Source/PythonScriptPlugin/Private/PyWrapperStruct.cpp
@@ -433,7 +433,7 @@ PyObject* FPyWrapperStruct::GetPropertyValue(FPyWrapperStruct* InSelf, const PyG
 	return PyGenUtil::GetPropertyValue(InSelf->ScriptStruct, InSelf->StructInstance, InPropDef, InPythonAttrName, (PyObject*)InSelf, *PyUtil::GetErrorContext(InSelf));
 }
 
-int FPyWrapperStruct::SetPropertyValue(FPyWrapperStruct* InSelf, PyObject* InValue, const PyGenUtil::FGeneratedWrappedProperty& InPropDef, const char* InPythonAttrName, const bool InNotifyChange, const uint64 InReadOnlyFlags)
+int FPyWrapperStruct::SetPropertyValue(FPyWrapperStruct* InSelf, PyObject* InValue, const PyGenUtil::FGeneratedWrappedProperty& InPropDef, const char* InPythonAttrName, const EPropertyAccessChangeNotifyMode InNotifyMode, const uint64 InReadOnlyFlags)
 {
 	if (!ValidateInternalState(InSelf))
 	{
@@ -447,7 +447,7 @@ int FPyWrapperStruct::SetPropertyValue(FPyWrapperStruct* InSelf, PyObject* InVal
 		OwnerIsTemplate = PropertyAccessUtil::IsObjectTemplate(OwnerObject);
 	}
 
-	const TUniquePtr<FPropertyAccessChangeNotify> ChangeNotify = InNotifyChange ? FPyWrapperOwnerContext((PyObject*)InSelf, InPropDef.Prop).BuildChangeNotify() : nullptr;
+	const TUniquePtr<FPropertyAccessChangeNotify> ChangeNotify = FPyWrapperOwnerContext((PyObject*)InSelf, InPropDef.Prop).BuildChangeNotify(InNotifyMode);
 	return PyGenUtil::SetPropertyValue(InSelf->ScriptStruct, InSelf->StructInstance, InValue, InPropDef, InPythonAttrName, ChangeNotify.Get(), InReadOnlyFlags, OwnerIsTemplate, *PyUtil::GetErrorContext(InSelf));
 }
 
@@ -1051,9 +1051,10 @@ PyTypeObject InitializePyWrapperStructType()
 
 			PyObject* PyNameObj = nullptr;
 			PyObject* PyValueObj = nullptr;
+			PyObject* PyNotifyModeObj = nullptr;
 
-			static const char *ArgsKwdList[] = { "name", "value", nullptr };
-			if (!PyArg_ParseTupleAndKeywords(InArgs, InKwds, "OO:set_editor_property", (char**)ArgsKwdList, &PyNameObj, &PyValueObj))
+			static const char *ArgsKwdList[] = { "name", "value", "notify_mode", nullptr };
+			if (!PyArg_ParseTupleAndKeywords(InArgs, InKwds, "OO|O:set_editor_property", (char**)ArgsKwdList, &PyNameObj, &PyValueObj, &PyNotifyModeObj))
 			{
 				return nullptr;
 			}
@@ -1065,6 +1066,15 @@ PyTypeObject InitializePyWrapperStructType()
 				return nullptr;
 			}
 
+			static const UEnum* PropertyAccessChangeNotifyModeEnum = FindObjectChecked<UEnum>(ANY_PACKAGE, TEXT("EPropertyAccessChangeNotifyMode"));
+
+			EPropertyAccessChangeNotifyMode NotifyMode = EPropertyAccessChangeNotifyMode::Default;
+			if (PyNotifyModeObj && !PyConversion::NativizeEnumEntry(PyNotifyModeObj, PropertyAccessChangeNotifyModeEnum, NotifyMode))
+			{
+				PyUtil::SetPythonError(PyExc_TypeError, InSelf, *FString::Printf(TEXT("Failed to convert 'notify_mode' (%s) to 'PropertyAccessChangeNotifyMode'"), *PyUtil::GetFriendlyTypename(PyNotifyModeObj)));
+				return nullptr;
+			}
+
 			const FName ResolvedName = FPyWrapperStructMetaData::ResolvePropertyName(InSelf, Name);
 			const UProperty* ResolvedProp = InSelf->ScriptStruct->FindPropertyByName(ResolvedName);
 			if (!ResolvedProp)
@@ -1093,7 +1103,7 @@ PyTypeObject InitializePyWrapperStructType()
 				WrappedPropDef.SetProperty(ResolvedProp);
 			}
 
-			const int Result = FPyWrapperStruct::SetPropertyValue(InSelf, PyValueObj, WrappedPropDef, TCHAR_TO_UTF8(*Name.ToString()), /*InNotifyChange*/true, PropertyAccessUtil::EditorReadOnlyFlags);
+			const int Result = FPyWrapperStruct::SetPropertyValue(InSelf, PyValueObj, WrappedPropDef, TCHAR_TO_UTF8(*Name.ToString()), NotifyMode, PropertyAccessUtil::EditorReadOnlyFlags);
 			if (Result != 0)
 			{
 				return nullptr;
@@ -1112,7 +1122,7 @@ PyTypeObject InitializePyWrapperStructType()
 		{ "assign", PyCFunctionCast(&FMethods::Assign), METH_VARARGS, "x.assign(object) -> None -- assign the value of this Unreal struct to value of the given object" },
 		{ "to_tuple", PyCFunctionCast(&FMethods::ToTuple), METH_NOARGS, "x.to_tuple() -> tuple -- break this Unreal struct into a tuple of its properties" },
 		{ "get_editor_property", PyCFunctionCast(&FMethods::GetEditorProperty), METH_VARARGS | METH_KEYWORDS, "x.get_editor_property(name) -> object -- get the value of any property visible to the editor" },
-		{ "set_editor_property", PyCFunctionCast(&FMethods::SetEditorProperty), METH_VARARGS | METH_KEYWORDS, "x.set_editor_property(name, value) -> None -- set the value of any property visible to the editor, ensuring that the pre/post change notifications are called" },
+		{ "set_editor_property", PyCFunctionCast(&FMethods::SetEditorProperty), METH_VARARGS | METH_KEYWORDS, "x.set_editor_property(name, value, notify_mode=PropertyAccessChangeNotifyMode.DEFAULT) -> None -- set the value of any property visible to the editor, ensuring that the pre/post change notifications are called" },
 		{ nullptr, nullptr, 0, nullptr }
 	};
 
diff --git a/Engine/Plugins/Experimental/PythonScriptPlugin/Source/PythonScriptPlugin/Private/PyWrapperStruct.h b/Engine/Plugins/Experimental/PythonScriptPlugin/Source/PythonScriptPlugin/Private/PyWrapperStruct.h
--- a/Engine/Plugins/Experimental/PythonScriptPlugin/Source/PythonScriptPlugin/Private/PyWrapperStruct.h
+++ b/Engine/Plugins/Experimental/PythonScriptPlugin/Source/PythonScriptPlugin/Private/PyWrapperStruct.h
@@ -65,7 +65,7 @@ struct FPyWrapperStruct : public FPyWrapperBase
 	static PyObject* GetPropertyValue(FPyWrapperStruct* InSelf, const PyGenUtil::FGeneratedWrappedProperty& InPropDef, const char* InPythonAttrName);
 
 	/** Set a property value on this instance (called via generated code) */
-	static int SetPropertyValue(FPyWrapperStruct* InSelf, PyObject* InValue, const PyGenUtil::FGeneratedWrappedProperty& InPropDef, const char* InPythonAttrName, const bool InNotifyChange = false, const uint64 InReadOnlyFlags = PropertyAccessUtil::RuntimeReadOnlyFlags);
+	static int SetPropertyValue(FPyWrapperStruct* InSelf, PyObject* InValue, const PyGenUtil::FGeneratedWrappedProperty& InPropDef, const char* InPythonAttrName, const EPropertyAccessChangeNotifyMode InNotifyMode = EPropertyAccessChangeNotifyMode::Never, const uint64 InReadOnlyFlags = PropertyAccessUtil::RuntimeReadOnlyFlags);
 
 	/** Call a make function on this instance (MakeStruct internal use only) */
 	static int CallMakeFunction_Impl(FPyWrapperStruct* InSelf, PyObject* InArgs, PyObject* InKwds, const PyGenUtil::FGeneratedWrappedFunction& InFuncDef);
diff --git a/Engine/Source/Editor/BlueprintGraph/Private/K2Node_EditorPropertyAccess.cpp b/Engine/Source/Editor/BlueprintGraph/Private/K2Node_EditorPropertyAccess.cpp
--- a/Engine/Source/Editor/BlueprintGraph/Private/K2Node_EditorPropertyAccess.cpp
+++ b/Engine/Source/Editor/BlueprintGraph/Private/K2Node_EditorPropertyAccess.cpp
@@ -2,8 +2,6 @@
 
 #include "K2Node_EditorPropertyAccess.h"
 #include "EdGraphSchema_K2.h"
-#include "K2Node_CallFunction.h"
-#include "K2Node_IfThenElse.h"
 #include "KismetCompiler.h"
 #include "Kismet2/BlueprintEditorUtils.h"
 #include "BlueprintNodeSpawner.h"
@@ -13,52 +11,12 @@
 
 #define LOCTEXT_NAMESPACE "K2Node_EditorPropertyAccess"
 
-namespace EditorPropertyAccessHelper
-{
-
-const FName ObjectPinName = "Object";
-const FName PropertyNamePinName = "PropertyName";
-const FName PropertyValuePinName = "PropertyValue";
-
-}
-
 void UK2Node_EditorPropertyAccessBase::AllocateDefaultPins()
 {
 	Super::AllocateDefaultPins();
 
-	// Add execution pins
-	CreatePin(EGPD_Input, UEdGraphSchema_K2::PC_Exec, UEdGraphSchema_K2::PN_Execute);
-	CreatePin(EGPD_Output, UEdGraphSchema_K2::PC_Exec, UEdGraphSchema_K2::PN_Then);
-
-	// Add Object pin
-	UEdGraphPin* ObjectPin = CreatePin(EGPD_Input, UEdGraphSchema_K2::PC_Object, UObject::StaticClass(), EditorPropertyAccessHelper::ObjectPinName);
-	SetPinToolTip(*ObjectPin, LOCTEXT("ObjectPinDescription", "The object you want to access a property value from"));
-
-	// Add Property Name pin
-	UEdGraphPin* PropertyNamePin = CreatePin(EGPD_Input, UEdGraphSchema_K2::PC_Name, EditorPropertyAccessHelper::PropertyNamePinName);
-	SetPinToolTip(*PropertyNamePin, LOCTEXT("PropertyNamePinDescription", "The name of the property to access from the object"));
-
-	// Add Property Value pin
-	AllocatePropertyValuePin();
-
-	// Add Result Pin
-	UEdGraphPin* ResultPin = CreatePin(EGPD_Output, UEdGraphSchema_K2::PC_Boolean, UEdGraphSchema_K2::PN_ReturnValue);
+	UEdGraphPin* ResultPin = FindPinChecked(UEdGraphSchema_K2::PN_ReturnValue, EGPD_Output);
 	ResultPin->PinFriendlyName = LOCTEXT("ResultPinFriendlyName", "Success?");
-	SetPinToolTip(*ResultPin, LOCTEXT("ResultPinDescription", "Whether the property value was found"));
-}
-
-void UK2Node_EditorPropertyAccessBase::SetPinToolTip(UEdGraphPin& MutatablePin, const FText& PinDescription) const
-{
-	MutatablePin.PinToolTip = UEdGraphSchema_K2::TypeToText(MutatablePin.PinType).ToString();
-
-	const UEdGraphSchema_K2* K2Schema = Cast<const UEdGraphSchema_K2>(GetSchema());
-	if (K2Schema)
-	{
-		MutatablePin.PinToolTip += TEXT(" ");
-		MutatablePin.PinToolTip += K2Schema->GetPinDisplayName(&MutatablePin).ToString();
-	}
-
-	MutatablePin.PinToolTip += FString(TEXT("\n")) + PinDescription.ToString();
 }
 
 void UK2Node_EditorPropertyAccessBase::GetMenuActions(FBlueprintActionDatabaseRegistrar& ActionRegistrar) const
@@ -81,11 +39,6 @@ void UK2Node_EditorPropertyAccessBase::GetMenuActions(FBlueprintActionDatabaseRe
 	}
 }
 
-FText UK2Node_EditorPropertyAccessBase::GetMenuCategory() const
-{
-	return FEditorCategoryUtils::GetCommonCategory(FCommonEditorCategory::Utilities);
-}
-
 bool UK2Node_EditorPropertyAccessBase::CanPasteHere(const UEdGraph* TargetGraph) const
 {
 	bool bCanPaste = Super::CanPasteHere(TargetGraph);
@@ -109,224 +62,34 @@ bool UK2Node_EditorPropertyAccessBase::IsActionFilteredOut(const FBlueprintActio
 	return bIsFilteredOut;
 }
 
-bool UK2Node_EditorPropertyAccessBase::IsConnectionDisallowed(const UEdGraphPin* MyPin, const UEdGraphPin* OtherPin, FString& OutReason) const
-{
-	if (MyPin->PinName == EditorPropertyAccessHelper::PropertyValuePinName && MyPin->PinType.PinCategory == UEdGraphSchema_K2::PC_Wildcard)
-	{
-		if (OtherPin->PinType.PinCategory == UEdGraphSchema_K2::PC_Exec)
-		{
-			OutReason = TEXT("Cannot connect the property value to an execution pin");
-			return true;
-		}
-	}
-
-	return false;
-}
-
-UEdGraphPin* UK2Node_EditorPropertyAccessBase::GetThenPin()const
-{
-	UEdGraphPin* Pin = FindPinChecked(UEdGraphSchema_K2::PN_Then);
-	check(Pin->Direction == EGPD_Output);
-	return Pin;
-}
-
-UEdGraphPin* UK2Node_EditorPropertyAccessBase::GetObjectPin() const
-{
-	UEdGraphPin* Pin = FindPinChecked(EditorPropertyAccessHelper::ObjectPinName);
-	check(Pin->Direction == EGPD_Input);
-	return Pin;
-}
-
-UEdGraphPin* UK2Node_EditorPropertyAccessBase::GetPropertyNamePin() const
-{
-	UEdGraphPin* Pin = FindPinChecked(EditorPropertyAccessHelper::PropertyNamePinName);
-	check(Pin->Direction == EGPD_Input);
-	return Pin;
-}
-
-UEdGraphPin* UK2Node_EditorPropertyAccessBase::GetResultPin() const
-{
-	UEdGraphPin* Pin = FindPinChecked(UEdGraphSchema_K2::PN_ReturnValue);
-	check(Pin->Direction == EGPD_Output);
-	return Pin;
-}
-
-void UK2Node_EditorPropertyAccessBase::RefreshPropertyValuePin()
-{
-	UEdGraphPin* PropertyValuePin = GetPropertyValuePin();
-
-	FEdGraphPinType NewPinType;
-	if (PropertyValuePin->LinkedTo.Num() > 0)
-	{
-		NewPinType = PropertyValuePin->LinkedTo[0]->PinType;
-	}
-	else
-	{
-		NewPinType.PinCategory = UEdGraphSchema_K2::PC_Wildcard;
-	}
-
-	if (PropertyValuePin->PinType != NewPinType)
-	{
-		PropertyValuePin->PinType = NewPinType;
-
-		GetGraph()->NotifyGraphChanged();
-
-		UBlueprint* Blueprint = GetBlueprint();
-		if (!Blueprint->bBeingCompiled)
-		{
-			FBlueprintEditorUtils::MarkBlueprintAsModified(Blueprint);
-			Blueprint->BroadcastChanged();
-		}
-	}
-}
-
-FSlateIcon UK2Node_EditorPropertyAccessBase::GetIconAndTint(FLinearColor& OutColor) const
-{
-	OutColor = GetNodeTitleColor();
-	static FSlateIcon Icon("EditorStyle", "Kismet.AllClasses.FunctionIcon");
-	return Icon;
-}
-
-void UK2Node_EditorPropertyAccessBase::PostReconstructNode()
-{
-	Super::PostReconstructNode();
-
-	RefreshPropertyValuePin();
-}
-
 void UK2Node_EditorPropertyAccessBase::EarlyValidation(class FCompilerResultsLog& MessageLog) const
 {
 	Super::EarlyValidation(MessageLog);
 
-	const UEdGraphPin* ObjectPin = GetObjectPin();
-	const UEdGraphPin* PropertyNamePin = GetPropertyNamePin();
-	const UEdGraphPin* PropertyValuePin = GetPropertyValuePin();
-	if (!ObjectPin || !PropertyNamePin || !PropertyValuePin)
-	{
-		MessageLog.Error(*LOCTEXT("MissingPins", "Missing pins in @@").ToString(), this);
-		return;
-	}
+	static const FName PN_Object = "Object";
+	static const FName PN_PropertyName = "PropertyName";
 
+	UEdGraphPin* ObjectPin = FindPinChecked(PN_Object, EGPD_Input);
 	if (ObjectPin->LinkedTo.Num() == 0 && !ObjectPin->DefaultObject)
 	{
 		MessageLog.Error(*LOCTEXT("UnsetObject", "No object set on @@").ToString(), this);
 	}
 
+	UEdGraphPin* PropertyNamePin = FindPinChecked(PN_PropertyName, EGPD_Input);
 	if (PropertyNamePin->LinkedTo.Num() == 0 && FName(*PropertyNamePin->DefaultValue).IsNone())
 	{
 		MessageLog.Error(*LOCTEXT("UnsetPropertyName", "No property name set on @@").ToString(), this);
 	}
 }
 
-void UK2Node_EditorPropertyAccessBase::NotifyPinConnectionListChanged(UEdGraphPin* Pin)
-{
-	Super::NotifyPinConnectionListChanged(Pin);
-
-	if (Pin->PinName == EditorPropertyAccessHelper::PropertyValuePinName)
-	{
-		RefreshPropertyValuePin();
-	}
-}
-
-void UK2Node_EditorPropertyAccessBase::ExpandNode(FKismetCompilerContext& CompilerContext, UEdGraph* SourceGraph)
-{
-	Super::ExpandNode(CompilerContext, SourceGraph);
-
-	const FName FunctionName = GetUnderlyingFunctionName();
-
-	// Add a CallFunction node for the underlying function
-	UK2Node_CallFunction* CallFunctionNode = CompilerContext.SpawnIntermediateNode<UK2Node_CallFunction>(this, SourceGraph);
-	CallFunctionNode->FunctionReference.SetExternalMember(FunctionName, UKismetSystemLibrary::StaticClass());
-	CallFunctionNode->AllocateDefaultPins();
-
-	// Re-wire the execution pins
-	CompilerContext.MovePinLinksToIntermediate(*GetExecPin(), *CallFunctionNode->GetExecPin());
-	CompilerContext.MovePinLinksToIntermediate(*GetThenPin(), *CallFunctionNode->GetThenPin());
-
-	// Re-wire the Object pin to the function input
-	{
-		UEdGraphPin* FunctionObjectPin = CallFunctionNode->FindPinChecked(TEXT("Object"));
-		CompilerContext.MovePinLinksToIntermediate(*GetObjectPin(), *FunctionObjectPin);
-	}
-
-	// Re-wire the PropertyName pin to the function input
-	{
-		UEdGraphPin* FunctionPropertyNamePin = CallFunctionNode->FindPinChecked(TEXT("PropertyName"));
-		CompilerContext.MovePinLinksToIntermediate(*GetPropertyNamePin(), *FunctionPropertyNamePin);
-	}
-
-	// Re-wire the ValuePtr pin to the function output, and copy its type over
-	{
-		UEdGraphPin* PropertyValuePin = GetPropertyValuePin();
-		UEdGraphPin* FunctionValuePtrPin = CallFunctionNode->FindPinChecked(TEXT("ValuePtr"));
-		CompilerContext.MovePinLinksToIntermediate(*PropertyValuePin, *FunctionValuePtrPin);
-		FunctionValuePtrPin->PinType = PropertyValuePin->PinType;
-	}
-
-	// Re-wire the result pin
-	{
-		UEdGraphPin* FunctionReturnPin = CallFunctionNode->FindPinChecked(UEdGraphSchema_K2::PN_ReturnValue);
-		CompilerContext.MovePinLinksToIntermediate(*GetResultPin(), *FunctionReturnPin);
-	}
-
-	// Disconnect this node
-	BreakAllNodeLinks();
-}
-
-void UK2Node_GetEditorProperty::AllocatePropertyValuePin()
-{
-	UEdGraphPin* PropertyValuePin = CreatePin(EGPD_Output, UEdGraphSchema_K2::PC_Wildcard, EditorPropertyAccessHelper::PropertyValuePinName);
-	SetPinToolTip(*PropertyValuePin, LOCTEXT("GetEditorProperty_PropertyValueDescription", "The returned property value, if found"));
-}
-
-UEdGraphPin* UK2Node_GetEditorProperty::GetPropertyValuePin() const
-{
-	UEdGraphPin* Pin = FindPinChecked(EditorPropertyAccessHelper::PropertyValuePinName);
-	check(Pin->Direction == EGPD_Output);
-	return Pin;
-}
-
-FText UK2Node_GetEditorProperty::GetNodeTitle(ENodeTitleType::Type TitleType) const
-{
-	return LOCTEXT("GetEditorProperty_NodeTitle", "Get Editor Property");
-}
-
-FText UK2Node_GetEditorProperty::GetTooltipText() const
-{
-	return LOCTEXT("GetEditorProperty_NodeTooltip", "Attempts to retrieve the value of a named property from the given object");
-}
-
-FName UK2Node_GetEditorProperty::GetUnderlyingFunctionName() const
-{
-	return GET_FUNCTION_NAME_CHECKED(UKismetSystemLibrary, GetEditorProperty);
-}
-
-void UK2Node_SetEditorProperty::AllocatePropertyValuePin()
-{
-	UEdGraphPin* PropertyValuePin = CreatePin(EGPD_Input, UEdGraphSchema_K2::PC_Wildcard, EditorPropertyAccessHelper::PropertyValuePinName);
-	SetPinToolTip(*PropertyValuePin, LOCTEXT("SetEditorProperty_PropertyValueDescription", "The property value to set"));
-}
-
-UEdGraphPin* UK2Node_SetEditorProperty::GetPropertyValuePin() const
-{
-	UEdGraphPin* Pin = FindPinChecked(EditorPropertyAccessHelper::PropertyValuePinName);
-	check(Pin->Direction == EGPD_Input);
-	return Pin;
-}
-
-FText UK2Node_SetEditorProperty::GetNodeTitle(ENodeTitleType::Type TitleType) const
-{
-	return LOCTEXT("SetEditorProperty_NodeTitle", "Set Editor Property");
-}
-
-FText UK2Node_SetEditorProperty::GetTooltipText() const
+UK2Node_GetEditorProperty::UK2Node_GetEditorProperty()
 {
-	return LOCTEXT("SetEditorProperty_NodeTooltip", "Attempts to set the value of a named property on the given object");
+	FunctionReference.SetExternalMember(GET_FUNCTION_NAME_CHECKED(UKismetSystemLibrary, GetEditorProperty), UKismetSystemLibrary::StaticClass());
 }
 
-FName UK2Node_SetEditorProperty::GetUnderlyingFunctionName() const
+UK2Node_SetEditorProperty::UK2Node_SetEditorProperty()
 {
-	return GET_FUNCTION_NAME_CHECKED(UKismetSystemLibrary, SetEditorProperty);
+	FunctionReference.SetExternalMember(GET_FUNCTION_NAME_CHECKED(UKismetSystemLibrary, SetEditorProperty), UKismetSystemLibrary::StaticClass());
 }
 
 #undef LOCTEXT_NAMESPACE
diff --git a/Engine/Source/Editor/BlueprintGraph/Private/K2Node_EditorPropertyAccess.h b/Engine/Source/Editor/BlueprintGraph/Private/K2Node_EditorPropertyAccess.h
--- a/Engine/Source/Editor/BlueprintGraph/Private/K2Node_EditorPropertyAccess.h
+++ b/Engine/Source/Editor/BlueprintGraph/Private/K2Node_EditorPropertyAccess.h
@@ -5,7 +5,7 @@
 
 #include "CoreMinimal.h"
 #include "Textures/SlateIcon.h"
-#include "K2Node.h"
+#include "K2Node_CallFunction.h"
 #include "K2Node_EditorPropertyAccess.generated.h"
 
 class UEdGraph;
@@ -13,56 +13,21 @@ class UK2Node_CallFunction;
 class FBlueprintActionDatabaseRegistrar;
 
 UCLASS(Abstract)
-class UK2Node_EditorPropertyAccessBase : public UK2Node
+class UK2Node_EditorPropertyAccessBase : public UK2Node_CallFunction
 {
 	GENERATED_BODY()
 
 public:
 	//~ Begin UEdGraphNode Interface.
 	virtual void AllocateDefaultPins() override;
-	virtual FSlateIcon GetIconAndTint(FLinearColor& OutColor) const override;
-	virtual void PostReconstructNode() override;
 	//~ End UEdGraphNode Interface.
 
 	//~ Begin UK2Node Interface
 	virtual void GetMenuActions(FBlueprintActionDatabaseRegistrar& ActionRegistrar) const override;
-	virtual FText GetMenuCategory() const override;
 	virtual bool CanPasteHere(const UEdGraph* TargetGraph) const override;
 	virtual bool IsActionFilteredOut(const class FBlueprintActionFilter& Filter) override;
-	virtual bool IsConnectionDisallowed(const UEdGraphPin* MyPin, const UEdGraphPin* OtherPin, FString& OutReason) const override;
 	virtual void EarlyValidation(class FCompilerResultsLog& MessageLog) const override;
-	virtual void NotifyPinConnectionListChanged(UEdGraphPin* Pin) override;
-	virtual void ExpandNode(class FKismetCompilerContext& CompilerContext, UEdGraph* SourceGraph) override;
 	//~ End UK2Node Interface
-
-	/** Get the then output pin */
-	UEdGraphPin* GetThenPin() const;
-	/** Get the object input pin */
-	UEdGraphPin* GetObjectPin() const;
-	/** Get the property name input pin */
-	UEdGraphPin* GetPropertyNamePin() const;
-	/** Get the property value output pin */
-	virtual UEdGraphPin* GetPropertyValuePin() const PURE_VIRTUAL(UK2Node_EditorPropertyAccessBase::GetPropertyValuePin, return nullptr;);
-	/** Get the result output pin */
-	UEdGraphPin* GetResultPin() const;
-
-protected:
-	/** Allocate the property value pin for this node */
-	virtual void AllocatePropertyValuePin() PURE_VIRTUAL(UK2Node_EditorPropertyAccessBase::AllocatePropertyValuePin, );
-
-	/** Get the function name from UKismetSystemLibrary that the CallFunction node should use for the underlying access function */
-	virtual FName GetUnderlyingFunctionName() const PURE_VIRTUAL(UK2Node_EditorPropertyAccessBase::GetUnderlyingFunctionName, return FName(););
-
-	/** Updates the type of the property value pin based on its connection type */
-	void RefreshPropertyValuePin();
-
-	/**
-	 * Takes the specified "MutatablePin" and sets its 'PinToolTip' field (according to the specified description)
-	 * 
-	 * @param   MutatablePin	The pin you want to set tool-tip text on
-	 * @param   PinDescription	A string describing the pin's purpose
-	 */
-	void SetPinToolTip(UEdGraphPin& MutatablePin, const FText& PinDescription) const;
 };
 
 UCLASS()
@@ -71,16 +36,7 @@ class UK2Node_GetEditorProperty : public UK2Node_EditorPropertyAccessBase
 	GENERATED_BODY()
 
 public:
-	//~ Begin UEdGraphNode Interface.
-	virtual FText GetNodeTitle(ENodeTitleType::Type TitleType) const override;
-	virtual FText GetTooltipText() const override;
-	//~ End UEdGraphNode Interface.
-
-	//~ Begin UK2Node_EditorPropertyAccessBase Interface.
-	virtual void AllocatePropertyValuePin() override;
-	virtual UEdGraphPin* GetPropertyValuePin() const override;
-	virtual FName GetUnderlyingFunctionName() const override;
-	//~ End UK2Node_EditorPropertyAccessBase Interface.
+	UK2Node_GetEditorProperty();
 };
 
 UCLASS()
@@ -89,14 +45,5 @@ class UK2Node_SetEditorProperty : public UK2Node_EditorPropertyAccessBase
 	GENERATED_BODY()
 
 public:
-	//~ Begin UEdGraphNode Interface.
-	virtual FText GetNodeTitle(ENodeTitleType::Type TitleType) const override;
-	virtual FText GetTooltipText() const override;
-	//~ End UEdGraphNode Interface.
-
-	//~ Begin UK2Node_EditorPropertyAccessBase Interface.
-	virtual void AllocatePropertyValuePin() override;
-	virtual UEdGraphPin* GetPropertyValuePin() const override;
-	virtual FName GetUnderlyingFunctionName() const override;
-	//~ End UK2Node_EditorPropertyAccessBase Interface.
+	UK2Node_SetEditorProperty();
 };
diff --git a/Engine/Source/Runtime/CoreUObject/Private/UObject/PropertyAccessUtil.cpp b/Engine/Source/Runtime/CoreUObject/Private/UObject/PropertyAccessUtil.cpp
--- a/Engine/Source/Runtime/CoreUObject/Private/UObject/PropertyAccessUtil.cpp
+++ b/Engine/Source/Runtime/CoreUObject/Private/UObject/PropertyAccessUtil.cpp
@@ -76,12 +76,12 @@ EPropertyAccessResultFlags CanGetPropertyValue(const UProperty* InProp)
 	return EPropertyAccessResultFlags::Success;
 }
 
-EPropertyAccessResultFlags SetPropertyValue_Object(const UProperty* InProp, UObject* InObject, const void* InSrcValue, const int32 InArrayIndex, const uint64 InReadOnlyFlags)
+EPropertyAccessResultFlags SetPropertyValue_Object(const UProperty* InProp, UObject* InObject, const void* InSrcValue, const int32 InArrayIndex, const uint64 InReadOnlyFlags, const EPropertyAccessChangeNotifyMode InNotifyMode)
 {
 	check(InObject->IsA(InProp->GetOwnerClass()));
-	return SetPropertyValue_InContainer(InProp, InObject, InSrcValue, InArrayIndex, InReadOnlyFlags, IsObjectTemplate(InObject), [InProp, InObject]()
+	return SetPropertyValue_InContainer(InProp, InObject, InSrcValue, InArrayIndex, InReadOnlyFlags, IsObjectTemplate(InObject), [InProp, InObject, InNotifyMode]()
 	{
-		return BuildBasicChangeNotify(InProp, InObject);
+		return BuildBasicChangeNotify(InProp, InObject, InNotifyMode);
 	});
 }
 
@@ -110,12 +110,14 @@ EPropertyAccessResultFlags SetPropertyValue_DirectSingle(const UProperty* InProp
 
 	return SetPropertyValue([InProp, InSrcValue, InDestValue](const FPropertyAccessChangeNotify* InChangeNotify)
 	{
-		if (!InProp->Identical(InSrcValue, InDestValue))
+		const bool bIdenticalValue = InProp->Identical(InSrcValue, InDestValue);
+		EmitPreChangeNotify(InChangeNotify, bIdenticalValue);
+		if (!bIdenticalValue)
 		{
-			EmitPreChangeNotify(InChangeNotify);
 			InProp->CopySingleValue(InDestValue, InSrcValue);
-			EmitPostChangeNotify(InChangeNotify);
 		}
+		EmitPostChangeNotify(InChangeNotify, bIdenticalValue);
+
 		return true;
 	}, InBuildChangeNotifyFunc);
 }
@@ -142,12 +144,14 @@ EPropertyAccessResultFlags SetPropertyValue_DirectComplete(const UProperty* InPr
 			return bIsIdentical;
 		};
 
-		if (!IdenticalComplete())
+		const bool bIdenticalValue = IdenticalComplete();
+		EmitPreChangeNotify(InChangeNotify, bIdenticalValue);
+		if (!bIdenticalValue)
 		{
-			EmitPreChangeNotify(InChangeNotify);
 			InProp->CopyCompleteValue(InDestValue, InSrcValue);
-			EmitPostChangeNotify(InChangeNotify);
 		}
+		EmitPostChangeNotify(InChangeNotify, bIdenticalValue);
+
 		return true;
 	}, InBuildChangeNotifyFunc);
 }
@@ -191,48 +195,57 @@ EPropertyAccessResultFlags CanSetPropertyValue(const UProperty* InProp, const ui
 	return EPropertyAccessResultFlags::Success;
 }
 
-void EmitPreChangeNotify(const FPropertyAccessChangeNotify* InChangeNotify)
+void EmitPreChangeNotify(const FPropertyAccessChangeNotify* InChangeNotify, const bool InIdenticalValue)
 {
 #if WITH_EDITOR
-	if (InChangeNotify)
+	if (InChangeNotify && InChangeNotify->NotifyMode != EPropertyAccessChangeNotifyMode::Never)
 	{
 		check(InChangeNotify->ChangedObject);
 
-		// Notify that a change is about to occur
-		InChangeNotify->ChangedObject->PreEditChange(const_cast<FEditPropertyChain&>(InChangeNotify->ChangedPropertyChain));
+		if (!InIdenticalValue || InChangeNotify->NotifyMode == EPropertyAccessChangeNotifyMode::Always)
+		{
+			// Notify that a change is about to occur
+			InChangeNotify->ChangedObject->PreEditChange(const_cast<FEditPropertyChain&>(InChangeNotify->ChangedPropertyChain));
+		}
 	}
 #endif
 }
 
-void EmitPostChangeNotify(const FPropertyAccessChangeNotify* InChangeNotify)
+void EmitPostChangeNotify(const FPropertyAccessChangeNotify* InChangeNotify, const bool InIdenticalValue)
 {
 #if WITH_EDITOR
-	if (InChangeNotify)
+	if (InChangeNotify && InChangeNotify->NotifyMode != EPropertyAccessChangeNotifyMode::Never)
 	{
 		check(InChangeNotify->ChangedObject);
 
-		// Notify that the change has occurred
-		FPropertyChangedEvent PropertyEvent(InChangeNotify->ChangedPropertyChain.GetActiveNode()->GetValue(), InChangeNotify->ChangeType, MakeArrayView(&InChangeNotify->ChangedObject, 1));
-		PropertyEvent.SetActiveMemberProperty(InChangeNotify->ChangedPropertyChain.GetActiveMemberNode()->GetValue());
-		FPropertyChangedChainEvent PropertyChainEvent(const_cast<FEditPropertyChain&>(InChangeNotify->ChangedPropertyChain), PropertyEvent);
-		InChangeNotify->ChangedObject->PostEditChangeChainProperty(PropertyChainEvent);
+		if (!InIdenticalValue || InChangeNotify->NotifyMode == EPropertyAccessChangeNotifyMode::Always)
+		{
+			// Notify that the change has occurred
+			FPropertyChangedEvent PropertyEvent(InChangeNotify->ChangedPropertyChain.GetActiveNode()->GetValue(), InChangeNotify->ChangeType, MakeArrayView(&InChangeNotify->ChangedObject, 1));
+			PropertyEvent.SetActiveMemberProperty(InChangeNotify->ChangedPropertyChain.GetActiveMemberNode()->GetValue());
+			FPropertyChangedChainEvent PropertyChainEvent(const_cast<FEditPropertyChain&>(InChangeNotify->ChangedPropertyChain), PropertyEvent);
+			InChangeNotify->ChangedObject->PostEditChangeChainProperty(PropertyChainEvent);
+		}
 	}
 #endif
 }
 
-TUniquePtr<FPropertyAccessChangeNotify> BuildBasicChangeNotify(const UProperty* InProp, const UObject* InObject)
+TUniquePtr<FPropertyAccessChangeNotify> BuildBasicChangeNotify(const UProperty* InProp, const UObject* InObject, const EPropertyAccessChangeNotifyMode InNotifyMode)
 {
 	check(InObject->IsA(InProp->GetOwnerClass()));
 #if WITH_EDITOR
-	TUniquePtr<FPropertyAccessChangeNotify> ChangeNotify = MakeUnique<FPropertyAccessChangeNotify>();
-	ChangeNotify->ChangedObject = const_cast<UObject*>(InObject);
-	ChangeNotify->ChangedPropertyChain.AddHead(const_cast<UProperty*>(InProp));
-	ChangeNotify->ChangedPropertyChain.SetActivePropertyNode(const_cast<UProperty*>(InProp));
-	ChangeNotify->ChangedPropertyChain.SetActiveMemberPropertyNode(const_cast<UProperty*>(InProp));
-	return ChangeNotify;
-#else
-	return nullptr;
+	if (InNotifyMode != EPropertyAccessChangeNotifyMode::Never)
+	{
+		TUniquePtr<FPropertyAccessChangeNotify> ChangeNotify = MakeUnique<FPropertyAccessChangeNotify>();
+		ChangeNotify->ChangedObject = const_cast<UObject*>(InObject);
+		ChangeNotify->ChangedPropertyChain.AddHead(const_cast<UProperty*>(InProp));
+		ChangeNotify->ChangedPropertyChain.SetActivePropertyNode(const_cast<UProperty*>(InProp));
+		ChangeNotify->ChangedPropertyChain.SetActiveMemberPropertyNode(const_cast<UProperty*>(InProp));
+		ChangeNotify->NotifyMode = InNotifyMode;
+		return ChangeNotify;
+	}
 #endif
+	return nullptr;
 }
 
 bool IsObjectTemplate(const UObject* InObject)
diff --git a/Engine/Source/Runtime/CoreUObject/Public/UObject/NoExportTypes.h b/Engine/Source/Runtime/CoreUObject/Public/UObject/NoExportTypes.h
--- a/Engine/Source/Runtime/CoreUObject/Public/UObject/NoExportTypes.h
+++ b/Engine/Source/Runtime/CoreUObject/Public/UObject/NoExportTypes.h
@@ -18,6 +18,7 @@
 #include "Misc/Timespan.h"
 
 #include "UObject/SoftObjectPath.h"
+#include "UObject/PropertyAccessUtil.h"
 
 #include "Math/InterpCurvePoint.h"
 #include "Math/UnitConversion.h"
@@ -339,6 +340,21 @@ enum class EUnit : uint8
 	Unspecified
 };
 
+/**
+ * Enum controlling when to emit property change notifications when setting a property value.
+ * @note Mirrored from PropertyAccessUtil.h
+ */
+UENUM(BlueprintType)
+enum class EPropertyAccessChangeNotifyMode : uint8
+{
+	/** Notify only when a value change has actually occurred */
+	Default,
+	/** Never notify that a value change has occurred */
+	Never,
+	/** Always notify that a value change has occurred, even if the value is unchanged */
+	Always,
+};
+
 /** A globally unique identifier (mirrored from Guid.h) */
 USTRUCT(immutable, noexport, BlueprintType)
 struct FGuid
diff --git a/Engine/Source/Runtime/CoreUObject/Public/UObject/PropertyAccessUtil.h b/Engine/Source/Runtime/CoreUObject/Public/UObject/PropertyAccessUtil.h
--- a/Engine/Source/Runtime/CoreUObject/Public/UObject/PropertyAccessUtil.h
+++ b/Engine/Source/Runtime/CoreUObject/Public/UObject/PropertyAccessUtil.h
@@ -35,6 +35,20 @@ enum class EPropertyAccessResultFlags : uint8
 };
 ENUM_CLASS_FLAGS(EPropertyAccessResultFlags);
 
+/**
+ * Enum controlling when to emit property change notifications when setting a property value.
+ * @note Mirrored in NoExportTypes.h for UHT.
+ */
+enum class EPropertyAccessChangeNotifyMode : uint8
+{
+	/** Notify only when a value change has actually occurred */
+	Default,
+	/** Never notify that a value change has occurred */
+	Never,
+	/** Always notify that a value change has occurred, even if the value is unchanged */
+	Always,
+};
+
 /**
  * Information needed to emit property change notifications when setting a property value.
  */
@@ -46,6 +60,8 @@ struct FPropertyAccessChangeNotify
 	UObject* ChangedObject = nullptr;
 	/** The chain of properties that are being changed */
 	FEditPropertyChain ChangedPropertyChain;
+	/** When to emit property change notifications */
+	EPropertyAccessChangeNotifyMode NotifyMode = EPropertyAccessChangeNotifyMode::Default;
 };
 
 /**
@@ -149,10 +165,11 @@ namespace PropertyAccessUtil
 	 * @param InSrcValue The value to set on the property.
 	 * @param InArrayIndex For fixed-size array properties denotes which index of the array to set, or INDEX_NONE to set the entire property.
 	 * @param InReadOnlyFlags Flags controlling which properties are considered read-only.
+	 * @param InNotifyMode When to emit property change notifications.
 	 *
 	 * @return Flags describing whether the set was successful.
 	 */
-	COREUOBJECT_API EPropertyAccessResultFlags SetPropertyValue_Object(const UProperty* InProp, UObject* InObject, const void* InSrcValue, const int32 InArrayIndex, const uint64 InReadOnlyFlags);
+	COREUOBJECT_API EPropertyAccessResultFlags SetPropertyValue_Object(const UProperty* InProp, UObject* InObject, const void* InSrcValue, const int32 InArrayIndex, const uint64 InReadOnlyFlags, const EPropertyAccessChangeNotifyMode InNotifyMode);
 
 	/**
 	 * High-level function for setting the value of a property on a property container (object or struct).
@@ -224,15 +241,17 @@ namespace PropertyAccessUtil
 	 * Low-level function called before modifying an object to notify that its value is about to change.
 	 *
 	 * @param InChangeNotify Information needed to emit property change notifications, or nullptr if no notifications are needed or possible.
+	 * @param InIdenticalValue True if the value being set was identical to the current value, false otherwise.
 	 */
-	COREUOBJECT_API void EmitPreChangeNotify(const FPropertyAccessChangeNotify* InChangeNotify);
+	COREUOBJECT_API void EmitPreChangeNotify(const FPropertyAccessChangeNotify* InChangeNotify, const bool InIdenticalValue);
 
 	/**
 	 * Low-level function called after modifying an object to notify that its value has changed.
 	 *
 	 * @param InChangeNotify Information needed to emit property change notifications, or nullptr if no notifications are needed or possible.
+	 * @param InIdenticalValue True if the value being set was identical to the current value, false otherwise.
 	 */
-	COREUOBJECT_API void EmitPostChangeNotify(const FPropertyAccessChangeNotify* InChangeNotify);
+	COREUOBJECT_API void EmitPostChangeNotify(const FPropertyAccessChangeNotify* InChangeNotify, const bool InIdenticalValue);
 
 	/**
 	 * Low-level function to build the basic information needed to emit property change notifications.
@@ -240,10 +259,11 @@ namespace PropertyAccessUtil
 	 *
 	 * @param InProp Property being modified.
 	 * @param InObject Object being modified.
+	 * @param InNotifyMode When to emit property change notifications.
 	 *
 	 * @return The information needed to emit property change notifications.
 	 */
-	COREUOBJECT_API TUniquePtr<FPropertyAccessChangeNotify> BuildBasicChangeNotify(const UProperty* InProp, const UObject* InObject);
+	COREUOBJECT_API TUniquePtr<FPropertyAccessChangeNotify> BuildBasicChangeNotify(const UProperty* InProp, const UObject* InObject, const EPropertyAccessChangeNotifyMode InNotifyMode);
 
 	/**
 	 * Low-level function for checking whether the given object instance is considered a template for property access.
diff --git a/Engine/Source/Runtime/Engine/Classes/Kismet/KismetSystemLibrary.h b/Engine/Source/Runtime/Engine/Classes/Kismet/KismetSystemLibrary.h
--- a/Engine/Source/Runtime/Engine/Classes/Kismet/KismetSystemLibrary.h
+++ b/Engine/Source/Runtime/Engine/Classes/Kismet/KismetSystemLibrary.h
@@ -12,6 +12,7 @@
 #include "UObject/Interface.h"
 #include "UObject/TextProperty.h"
 #include "UObject/SoftObjectPtr.h"
+#include "UObject/PropertyAccessUtil.h"
 #include "Engine/LatentActionManager.h"
 #include "Kismet/BlueprintFunctionLibrary.h"
 #include "Engine/CollisionProfile.h"
@@ -1736,16 +1737,33 @@ class ENGINE_API UKismetSystemLibrary : public UBlueprintFunctionLibrary
 	// --- Property Access ---------------------------
 
 #if WITH_EDITOR
-	/** Read the value of a named property on the given object to the given value pointer */
-    UFUNCTION(BlueprintCallable, CustomThunk, Category = "Properties", meta=(CustomStructureParam="ValuePtr", BlueprintInternalUseOnly="true"))
-    static bool GetEditorProperty(UObject* Object, const FName PropertyName, int32& ValuePtr);
+	/**
+	 * Attempts to retrieve the value of a named property from the given object.
+	 *
+	 * @param Object The object you want to retrieve a property value from.
+	 * @param PropertyName The name of the object property to retrieve the value from.
+	 * @param PropertyValue The retrieved property value, if found.
+	 *
+	 * @return Whether the property value was found and correctly retrieved.
+	 */
+    UFUNCTION(BlueprintCallable, CustomThunk, Category = "Utilities", meta=(CustomStructureParam="PropertyValue", BlueprintInternalUseOnly="true"))
+    static bool GetEditorProperty(UObject* Object, const FName PropertyName, int32& PropertyValue);
 	static bool Generic_GetEditorProperty(const UObject* Object, const UProperty* Property, void* ValuePtr);
 	DECLARE_FUNCTION(execGetEditorProperty);
 
-	/** Write the value of the given value pointer to a named property on the given object */
-    UFUNCTION(BlueprintCallable, CustomThunk, Category = "Properties", meta=(CustomStructureParam="ValuePtr", BlueprintInternalUseOnly="true"))
-    static bool SetEditorProperty(UObject* Object, const FName PropertyName, const int32& ValuePtr);
-	static bool Generic_SetEditorProperty(UObject* Object, const UProperty* Property, const void* ValuePtr);
+	/**
+	 * Attempts to set the value of a named property on the given object.
+	 *
+	 * @param Object The object you want to set a property value on.
+	 * @param PropertyName The name of the object property to set the value of.
+	 * @param PropertyValue The property value to set.
+	 * @param ChangeNotifyMode When to emit property change notifications.
+	 *
+	 * @return Whether the property value was found and correctly set.
+	 */
+    UFUNCTION(BlueprintCallable, CustomThunk, Category = "Utilities", meta=(CustomStructureParam="PropertyValue", AdvancedDisplay="ChangeNotifyMode", BlueprintInternalUseOnly="true"))
+    static bool SetEditorProperty(UObject* Object, const FName PropertyName, const int32& PropertyValue, const EPropertyAccessChangeNotifyMode ChangeNotifyMode);
+	static bool Generic_SetEditorProperty(UObject* Object, const UProperty* Property, const void* ValuePtr, const EPropertyAccessChangeNotifyMode ChangeNotifyMode);
 	DECLARE_FUNCTION(execSetEditorProperty);
 #endif
 
diff --git a/Engine/Source/Runtime/Engine/Private/KismetSystemLibrary.cpp b/Engine/Source/Runtime/Engine/Private/KismetSystemLibrary.cpp
--- a/Engine/Source/Runtime/Engine/Private/KismetSystemLibrary.cpp
+++ b/Engine/Source/Runtime/Engine/Private/KismetSystemLibrary.cpp
@@ -10,7 +10,6 @@
 #include "Misc/Paths.h"
 #include "Misc/RuntimeErrors.h"
 #include "UObject/GCObject.h"
-#include "UObject/PropertyAccessUtil.h"
 #include "EngineGlobals.h"
 #include "Components/ActorComponent.h"
 #include "TimerManager.h"
@@ -2663,7 +2662,7 @@ bool UKismetSystemLibrary::IsUnattended()
 
 #if WITH_EDITOR
 
-bool UKismetSystemLibrary::GetEditorProperty(UObject* Object, const FName PropertyName, int32& OutValuePtr)
+bool UKismetSystemLibrary::GetEditorProperty(UObject* Object, const FName PropertyName, int32& PropertyValue)
 {
 	// We should never hit this! Stubbed to avoid NoExport on the class.
 	check(0);
@@ -2744,16 +2743,16 @@ DEFINE_FUNCTION(UKismetSystemLibrary::execGetEditorProperty)
 	*(bool*)RESULT_PARAM = bResult;
 }
 
-bool UKismetSystemLibrary::SetEditorProperty(UObject* Object, const FName PropertyName, const int32& OutValuePtr)
+bool UKismetSystemLibrary::SetEditorProperty(UObject* Object, const FName PropertyName, const int32& PropertyValue, const EPropertyAccessChangeNotifyMode ChangeNotifyMode)
 {
 	// We should never hit this! Stubbed to avoid NoExport on the class.
 	check(0);
 	return false;
 }
 
-bool UKismetSystemLibrary::Generic_SetEditorProperty(UObject* Object, const UProperty* Property, const void* ValuePtr)
+bool UKismetSystemLibrary::Generic_SetEditorProperty(UObject* Object, const UProperty* Property, const void* ValuePtr, const EPropertyAccessChangeNotifyMode ChangeNotifyMode)
 {
-	const EPropertyAccessResultFlags AccessResult = PropertyAccessUtil::SetPropertyValue_Object(Property, Object, ValuePtr, INDEX_NONE, PropertyAccessUtil::EditorReadOnlyFlags);
+	const EPropertyAccessResultFlags AccessResult = PropertyAccessUtil::SetPropertyValue_Object(Property, Object, ValuePtr, INDEX_NONE, PropertyAccessUtil::EditorReadOnlyFlags, ChangeNotifyMode);
 
 	if (EnumHasAnyFlags(AccessResult, EPropertyAccessResultFlags::PermissionDenied))
 	{
@@ -2803,6 +2802,8 @@ DEFINE_FUNCTION(UKismetSystemLibrary::execSetEditorProperty)
 	const UProperty* ValueProp = Stack.MostRecentProperty;
 	const void* ValuePtr = Stack.MostRecentPropertyAddress;
 
+	P_GET_ENUM(EPropertyAccessChangeNotifyMode, ChangeNotifyMode);
+
 	P_FINISH;
 
 	if (!ValueProp || !ValuePtr)
@@ -2831,7 +2832,7 @@ DEFINE_FUNCTION(UKismetSystemLibrary::execSetEditorProperty)
 		if (ObjectProp && ObjectProp->GetClass() == ValueProp->GetClass()) // Compare the classes as these must be an *exact* match as the read is low-level and without property coercion
 		{
 			P_NATIVE_BEGIN;
-			bResult = Generic_SetEditorProperty(Object, ObjectProp, ValuePtr);
+			bResult = Generic_SetEditorProperty(Object, ObjectProp, ValuePtr, ChangeNotifyMode);
 			P_NATIVE_END;
 		}
 		else
