Implemented set_editor_properties for Python wrapped structs

This will attempt to notify the nearest object in the Python owner context stack, and is otherwise identical to the implementation for wrapped objects (see 10265342).

#rnx

diff --git a/Engine/Plugins/Experimental/PythonScriptPlugin/Source/PythonScriptPlugin/Private/PyWrapperOwnerContext.cpp b/Engine/Plugins/Experimental/PythonScriptPlugin/Source/PythonScriptPlugin/Private/PyWrapperOwnerContext.cpp
--- a/Engine/Plugins/Experimental/PythonScriptPlugin/Source/PythonScriptPlugin/Private/PyWrapperOwnerContext.cpp
+++ b/Engine/Plugins/Experimental/PythonScriptPlugin/Source/PythonScriptPlugin/Private/PyWrapperOwnerContext.cpp
@@ -118,4 +118,31 @@ TUniquePtr<FPropertyAccessChangeNotify> FPyWrapperOwnerContext::BuildChangeNotif
 	return nullptr;
 }
 
+UObject* FPyWrapperOwnerContext::FindChangeNotifyObject() const
+{
+	FPyWrapperOwnerContext OwnerContext = *this;
+	while (OwnerContext.HasOwner())
+	{
+		PyObject* PyObj = OwnerContext.GetOwnerObject();
+
+		if (PyObject_IsInstance(PyObj, (PyObject*)&PyWrapperObjectType) == 1)
+		{
+			// Found an object, this is the end of the chain
+			return ((FPyWrapperObject*)PyObj)->ObjectInstance;
+		}
+
+		if (PyObject_IsInstance(PyObj, (PyObject*)&PyWrapperStructType) == 1)
+		{
+			// Found a struct, recurse up the chain
+			OwnerContext = ((FPyWrapperStruct*)PyObj)->OwnerContext;
+			continue;
+		}
+
+		// Unknown object type - just bail
+		break;
+	}
+
+	return nullptr;
+}
+
 #endif	// WITH_PYTHON
diff --git a/Engine/Plugins/Experimental/PythonScriptPlugin/Source/PythonScriptPlugin/Private/PyWrapperOwnerContext.h b/Engine/Plugins/Experimental/PythonScriptPlugin/Source/PythonScriptPlugin/Private/PyWrapperOwnerContext.h
--- a/Engine/Plugins/Experimental/PythonScriptPlugin/Source/PythonScriptPlugin/Private/PyWrapperOwnerContext.h
+++ b/Engine/Plugins/Experimental/PythonScriptPlugin/Source/PythonScriptPlugin/Private/PyWrapperOwnerContext.h
@@ -40,6 +40,9 @@ public:
 	/** Build the property change notify that corresponds to this owner context, or null if this owner context shouldn't emit change notifications */
 	TUniquePtr<FPropertyAccessChangeNotify> BuildChangeNotify(const EPropertyAccessChangeNotifyMode InNotifyMode) const;
 
+	/** Walk the owner context chain to find a UObject owner instance that should receive change notifications (if any) */
+	UObject* FindChangeNotifyObject() const;
+
 private:
 	/** The Python object that owns the instance being wrapped (if any) */
 	FPyObjectPtr OwnerObject;
diff --git a/Engine/Plugins/Experimental/PythonScriptPlugin/Source/PythonScriptPlugin/Private/PyWrapperStruct.cpp b/Engine/Plugins/Experimental/PythonScriptPlugin/Source/PythonScriptPlugin/Private/PyWrapperStruct.cpp
--- a/Engine/Plugins/Experimental/PythonScriptPlugin/Source/PythonScriptPlugin/Private/PyWrapperStruct.cpp
+++ b/Engine/Plugins/Experimental/PythonScriptPlugin/Source/PythonScriptPlugin/Private/PyWrapperStruct.cpp
@@ -8,6 +8,7 @@
 #include "UObject/Package.h"
 #include "UObject/Class.h"
 #include "UObject/PropertyPortFlags.h"
+#include "Misc/ScopeExit.h"
 #include "Templates/Casts.h"
 #include "Engine/UserDefinedStruct.h"
 
@@ -989,6 +990,40 @@
 			return FPyWrapperStruct::BreakStruct(InSelf);
 		}
 
+		static PyGenUtil::FGeneratedWrappedProperty FindEditorPropertyImpl(FPyWrapperStruct* InSelf, const FName InPythonPropName)
+		{
+			PyGenUtil::FGeneratedWrappedProperty WrappedPropDef;
+
+			const FName ResolvedName = FPyWrapperStructMetaData::ResolvePropertyName(InSelf, InPythonPropName);
+			const UProperty* ResolvedProp = InSelf->ScriptStruct->FindPropertyByName(ResolvedName);
+			if (!ResolvedProp)
+			{
+				PyUtil::SetPythonError(PyExc_Exception, InSelf, *FString::Printf(TEXT("Failed to find property '%s' for attribute '%s' on '%s'"), *ResolvedName.ToString(), *InPythonPropName.ToString(), *InSelf->ScriptStruct->GetName()));
+				return WrappedPropDef;
+			}
+
+			TOptional<FString> PropDeprecationMessage;
+			{
+				FString PropDeprecationMessageStr;
+				if (FPyWrapperStructMetaData::IsPropertyDeprecated(InSelf, InPythonPropName, &PropDeprecationMessageStr))
+				{
+					PropDeprecationMessage = MoveTemp(PropDeprecationMessageStr);
+				}
+			}
+
+			if (PropDeprecationMessage.IsSet())
+			{
+				WrappedPropDef.SetProperty(ResolvedProp, PyGenUtil::FGeneratedWrappedProperty::SPF_None);
+				WrappedPropDef.DeprecationMessage = MoveTemp(PropDeprecationMessage);
+			}
+			else
+			{
+				WrappedPropDef.SetProperty(ResolvedProp);
+			}
+
+			return WrappedPropDef;
+		}
+
 		static PyObject* GetEditorProperty(FPyWrapperStruct* InSelf, PyObject* InArgs, PyObject* InKwds)
 		{
 			if (!FPyWrapperStruct::ValidateInternalState(InSelf))
@@ -1011,33 +1046,11 @@
 				return nullptr;
 			}
 
-			const FName ResolvedName = FPyWrapperStructMetaData::ResolvePropertyName(InSelf, Name);
-			const UProperty* ResolvedProp = InSelf->ScriptStruct->FindPropertyByName(ResolvedName);
-			if (!ResolvedProp)
+			const PyGenUtil::FGeneratedWrappedProperty WrappedPropDef = FindEditorPropertyImpl(InSelf, Name);
+			if (!WrappedPropDef.Prop)
 			{
-				PyUtil::SetPythonError(PyExc_Exception, InSelf, *FString::Printf(TEXT("Failed to find property '%s' for attribute '%s' on '%s'"), *ResolvedName.ToString(), *Name.ToString(), *InSelf->ScriptStruct->GetName()));
 				return nullptr;
 			}
-			
-			TOptional<FString> PropDeprecationMessage;
-			{
-				FString PropDeprecationMessageStr;
-				if (FPyWrapperStructMetaData::IsPropertyDeprecated(InSelf, Name, &PropDeprecationMessageStr))
-				{
-					PropDeprecationMessage = MoveTemp(PropDeprecationMessageStr);
-				}
-			}
 
-			PyGenUtil::FGeneratedWrappedProperty WrappedPropDef;
-			if (PropDeprecationMessage.IsSet())
-			{
-				WrappedPropDef.SetProperty(ResolvedProp, PyGenUtil::FGeneratedWrappedProperty::SPF_None);
-				WrappedPropDef.DeprecationMessage = MoveTemp(PropDeprecationMessage);
-			}
-			else
-			{
-				WrappedPropDef.SetProperty(ResolvedProp);
-			}
-
 			return FPyWrapperStruct::GetPropertyValue(InSelf, WrappedPropDef, TCHAR_TO_UTF8(*Name.ToString()));
 		}
@@ -1075,34 +1088,12 @@
 				return nullptr;
 			}
 
-			const FName ResolvedName = FPyWrapperStructMetaData::ResolvePropertyName(InSelf, Name);
-			const UProperty* ResolvedProp = InSelf->ScriptStruct->FindPropertyByName(ResolvedName);
-			if (!ResolvedProp)
+			const PyGenUtil::FGeneratedWrappedProperty WrappedPropDef = FindEditorPropertyImpl(InSelf, Name);
+			if (!WrappedPropDef.Prop)
 			{
-				PyUtil::SetPythonError(PyExc_Exception, InSelf, *FString::Printf(TEXT("Failed to find property '%s' for attribute '%s' on '%s'"), *ResolvedName.ToString(), *Name.ToString(), *InSelf->ScriptStruct->GetName()));
 				return nullptr;
 			}
 
-			TOptional<FString> PropDeprecationMessage;
-			{
-				FString PropDeprecationMessageStr;
-				if (FPyWrapperStructMetaData::IsPropertyDeprecated(InSelf, Name, &PropDeprecationMessageStr))
-				{
-					PropDeprecationMessage = MoveTemp(PropDeprecationMessageStr);
-				}
-			}
-
-			PyGenUtil::FGeneratedWrappedProperty WrappedPropDef;
-			if (PropDeprecationMessage.IsSet())
-			{
-				WrappedPropDef.SetProperty(ResolvedProp, PyGenUtil::FGeneratedWrappedProperty::SPF_None);
-				WrappedPropDef.DeprecationMessage = MoveTemp(PropDeprecationMessage);
-			}
-			else
-			{
-				WrappedPropDef.SetProperty(ResolvedProp);
-			}
-
 			const int Result = FPyWrapperStruct::SetPropertyValue(InSelf, PyValueObj, WrappedPropDef, TCHAR_TO_UTF8(*Name.ToString()), NotifyMode, PropertyAccessUtil::EditorReadOnlyFlags);
 			if (Result != 0)
 			{
@@ -1111,4 +1102,93 @@
 
 			Py_RETURN_NONE;
 		}
+
+		static PyObject* SetEditorProperties(FPyWrapperStruct* InSelf, PyObject* InArgs, PyObject* InKwds)
+		{
+			if (!FPyWrapperStruct::ValidateInternalState(InSelf))
+			{
+				return nullptr;
+			}
+
+			PyObject* PyPropertyInfoDict = nullptr;
+			if (!PyArg_ParseTuple(InArgs, "O:set_editor_properties", &PyPropertyInfoDict))
+			{
+				return nullptr;
+			}
+
+			if (!PyUtil::IsMappingType(PyPropertyInfoDict))
+			{
+				PyUtil::SetPythonError(PyExc_TypeError, InSelf, *FString::Printf(TEXT("'property_info' (%s) is expected to be a mapping type (eg, dict)"), *PyUtil::GetFriendlyTypename(PyPropertyInfoDict)));
+				return nullptr;
+			}
+
+			// Build up the list of properties and their new values
+			typedef TTuple<FName, PyGenUtil::FGeneratedWrappedProperty, FPyObjectPtr> FPropertyInfoPair;
+			TArray<FPropertyInfoPair> PropertyInfos;
+			{
+				FPyObjectPtr PyPropertyInfoDictIter = FPyObjectPtr::StealReference(PyObject_GetIter(PyPropertyInfoDict));
+				if (PyPropertyInfoDictIter)
+				{
+					FPyObjectPtr PyKeyItem;
+					while ((PyKeyItem = FPyObjectPtr::StealReference(PyIter_Next(PyPropertyInfoDictIter))))
+					{
+						FName Name;
+						if (!PyConversion::Nativize(PyKeyItem, Name))
+						{
+							PyUtil::SetPythonError(PyExc_TypeError, InSelf, *FString::Printf(TEXT("Failed to convert dict key (%s) to 'Name'"), *PyUtil::GetFriendlyTypename(PyKeyItem)));
+							return nullptr;
+						}
+
+						PyGenUtil::FGeneratedWrappedProperty WrappedPropDef = FindEditorPropertyImpl(InSelf, Name);
+						if (!WrappedPropDef.Prop)
+						{
+							return nullptr;
+						}
+
+						FPyObjectPtr PyValueItem = FPyObjectPtr::StealReference(PyObject_GetItem(PyPropertyInfoDict, PyKeyItem));
+						if (ensure(PyValueItem))
+						{
+							PropertyInfos.Add(MakeTuple(Name, MoveTemp(WrappedPropDef), MoveTemp(PyValueItem)));
+						}
+					}
+
+					// Iteration error?
+					if (PyErr_Occurred())
+					{
+						return nullptr;
+					}
+				}
+			}
+
+			// At this point we know that every property we were asked to set was found, but not that the values are going to be compatible
+			// Checking value coercion is tricky without actually doing the conversion, so we'll assume that this point that we need to do the change notifies
+			UObject* OwnerObjectInstance = InSelf->OwnerContext.FindChangeNotifyObject();
+			if (OwnerObjectInstance)
+			{
+				OwnerObjectInstance->PreEditChange(nullptr);
+			}
+			ON_SCOPE_EXIT
+			{
+				if (OwnerObjectInstance)
+				{
+					OwnerObjectInstance->PostEditChange();
+				}
+			};
+
+			// Try and set the value of each property
+			for (const FPropertyInfoPair& PropertyInfo : PropertyInfos)
+			{
+				const FName Name = PropertyInfo.Get<0>();
+				const PyGenUtil::FGeneratedWrappedProperty& WrappedPropDef = PropertyInfo.Get<1>();
+				PyObject* PyValueObj = PropertyInfo.Get<2>().GetPtr();
+
+				const int Result = FPyWrapperStruct::SetPropertyValue(InSelf, PyValueObj, WrappedPropDef, TCHAR_TO_UTF8(*Name.ToString()), EPropertyAccessChangeNotifyMode::Never, PropertyAccessUtil::EditorReadOnlyFlags);
+				if (Result != 0)
+				{
+					return nullptr;
+				}
+			}
+
+			Py_RETURN_NONE;
+		}
 	};
@@ -1123,6 +1203,7 @@
 		{ "to_tuple", PyCFunctionCast(&FMethods::ToTuple), METH_NOARGS, "x.to_tuple() -> tuple -- break this Unreal struct into a tuple of its properties" },
 		{ "get_editor_property", PyCFunctionCast(&FMethods::GetEditorProperty), METH_VARARGS | METH_KEYWORDS, "x.get_editor_property(name) -> object -- get the value of any property visible to the editor" },
 		{ "set_editor_property", PyCFunctionCast(&FMethods::SetEditorProperty), METH_VARARGS | METH_KEYWORDS, "x.set_editor_property(name, value, notify_mode=PropertyAccessChangeNotifyMode.DEFAULT) -> None -- set the value of any property visible to the editor, ensuring that the pre/post change notifications are called" },
+		{ "set_editor_properties", PyCFunctionCast(&FMethods::SetEditorProperties), METH_VARARGS, "x.set_editor_properties(property_info) -> None -- set the value of any properties visible to the editor (from a name->value dict), ensuring that the pre/post change notifications are called" },
 		{ nullptr, nullptr, 0, nullptr }
 	};
 
