PySlate now checks that FSlateApplication is available rather than crashing

diff --git a/Engine/Plugins/Experimental/PythonScriptPlugin/Source/PythonScriptPlugin/Private/PySlate.cpp b/Engine/Plugins/Experimental/PythonScriptPlugin/Source/PythonScriptPlugin/Private/PySlate.cpp
--- a/Engine/Plugins/Experimental/PythonScriptPlugin/Source/PythonScriptPlugin/Private/PySlate.cpp
+++ b/Engine/Plugins/Experimental/PythonScriptPlugin/Source/PythonScriptPlugin/Private/PySlate.cpp
@@ -78,7 +78,12 @@ PyObject* RegisterSlatePreTickCallback(PyObject* InSelf, PyObject* InArgs)
 	}
 	check(PyObj);
 
-	return (PyObject*)PySlateUtil::RegisterSlateTickCallback(FSlateApplication::Get().OnPreTick(), PyObj);
+	if (FSlateApplication::IsInitialized())
+	{
+		return (PyObject*)PySlateUtil::RegisterSlateTickCallback(FSlateApplication::Get().OnPreTick(), PyObj);
+	}
+
+	return (PyObject*)FPyDelegateHandle::CreateInstance(FDelegateHandle());
 }
 
 PyObject* UnregisterSlatePreTickCallback(PyObject* InSelf, PyObject* InArgs)
@@ -90,7 +95,7 @@ PyObject* UnregisterSlatePreTickCallback(PyObject* InSelf, PyObject* InArgs)
 	}
 	check(PyObj);
 
-	if (!PySlateUtil::UnregisterSlateTickCallback(FSlateApplication::Get().OnPreTick(), PyObj))
+	if (FSlateApplication::IsInitialized() && !PySlateUtil::UnregisterSlateTickCallback(FSlateApplication::Get().OnPreTick(), PyObj))
 	{
 		PyUtil::SetPythonError(PyExc_TypeError, TEXT("unregister_slate_pre_tick_callback"), *FString::Printf(TEXT("Failed to convert argument '%s' to '_DelegateHandle'"), *PyUtil::GetFriendlyTypename(PyObj)));
 		return nullptr;
@@ -108,7 +113,12 @@ PyObject* RegisterSlatePostTickCallback(PyObject* InSelf, PyObject* InArgs)
 	}
 	check(PyObj);
 
-	return (PyObject*)PySlateUtil::RegisterSlateTickCallback(FSlateApplication::Get().OnPostTick(), PyObj);
+	if (FSlateApplication::IsInitialized())
+	{
+		return (PyObject*)PySlateUtil::RegisterSlateTickCallback(FSlateApplication::Get().OnPostTick(), PyObj);
+	}
+
+	return (PyObject*)FPyDelegateHandle::CreateInstance(FDelegateHandle());
 }
 
 PyObject* UnregisterSlatePostTickCallback(PyObject* InSelf, PyObject* InArgs)
@@ -120,7 +130,7 @@ PyObject* UnregisterSlatePostTickCallback(PyObject* InSelf, PyObject* InArgs)
 	}
 	check(PyObj);
 
-	if (!PySlateUtil::UnregisterSlateTickCallback(FSlateApplication::Get().OnPostTick(), PyObj))
+	if (FSlateApplication::IsInitialized() && !PySlateUtil::UnregisterSlateTickCallback(FSlateApplication::Get().OnPostTick(), PyObj))
 	{
 		PyUtil::SetPythonError(PyExc_TypeError, TEXT("unregister_slate_post_tick_callback"), *FString::Printf(TEXT("Failed to convert argument '%s' to '_DelegateHandle'"), *PyUtil::GetFriendlyTypename(PyObj)));
 		return nullptr;
@@ -154,12 +164,15 @@ PyObject* ParentExternalWindowToSlate(PyObject* InSelf, PyObject* InArgs)
 		return nullptr;
 	}
 
-	const void* SlateParentWindowHandle = FSlateApplication::Get().FindBestParentWindowHandleForDialogs(nullptr, ParentWindowSearchMethod);
-	if (SlateParentWindowHandle && ExternalWindowHandle)
+	if (FSlateApplication::IsInitialized())
 	{
+		const void* SlateParentWindowHandle = FSlateApplication::Get().FindBestParentWindowHandleForDialogs(nullptr, ParentWindowSearchMethod);
+		if (SlateParentWindowHandle && ExternalWindowHandle)
+		{
 #if PLATFORM_WINDOWS
-		::SetWindowLongPtr((HWND)ExternalWindowHandle, -8/*GWL_HWNDPARENT*/, (LONG_PTR)SlateParentWindowHandle);
+			::SetWindowLongPtr((HWND)ExternalWindowHandle, -8/*GWL_HWNDPARENT*/, (LONG_PTR)SlateParentWindowHandle);
 #endif
+		}
 	}
 
 	Py_RETURN_NONE;
