Clean-up native module references when shutting down Python

diff --git a/Engine/Plugins/Experimental/PythonScriptPlugin/Source/PythonScriptPlugin/Private/PyCore.cpp b/Engine/Plugins/Experimental/PythonScriptPlugin/Source/PythonScriptPlugin/Private/PyCore.cpp
--- a/Engine/Plugins/Experimental/PythonScriptPlugin/Source/PythonScriptPlugin/Private/PyCore.cpp
+++ b/Engine/Plugins/Experimental/PythonScriptPlugin/Source/PythonScriptPlugin/Private/PyCore.cpp
@@ -1765,6 +1765,10 @@ void InitializeModule()
 	FPyWrapperTypeRegistry::Get().RegisterNativePythonModule(MoveTemp(NativePythonModule));
 }
 
+void ShutdownModule()
+{
+}
+
 FPythonLogCapture& GetPythonLogCapture()
 {
 	static FPythonLogCapture PythonLogCapture;
diff --git a/Engine/Plugins/Experimental/PythonScriptPlugin/Source/PythonScriptPlugin/Private/PyCore.h b/Engine/Plugins/Experimental/PythonScriptPlugin/Source/PythonScriptPlugin/Private/PyCore.h
--- a/Engine/Plugins/Experimental/PythonScriptPlugin/Source/PythonScriptPlugin/Private/PyCore.h
+++ b/Engine/Plugins/Experimental/PythonScriptPlugin/Source/PythonScriptPlugin/Private/PyCore.h
@@ -373,6 +373,7 @@ typedef TPyPtr<FPyUFunctionDef> FPyUFunctionDefPtr;
 namespace PyCore
 {
 	void InitializeModule();
+	void ShutdownModule();
 
 	DECLARE_MULTICAST_DELEGATE_TwoParams(FPythonLogCapture, EPythonLogOutputType, const TCHAR*);
 	FPythonLogCapture& GetPythonLogCapture();
diff --git a/Engine/Plugins/Experimental/PythonScriptPlugin/Source/PythonScriptPlugin/Private/PyEditor.cpp b/Engine/Plugins/Experimental/PythonScriptPlugin/Source/PythonScriptPlugin/Private/PyEditor.cpp
--- a/Engine/Plugins/Experimental/PythonScriptPlugin/Source/PythonScriptPlugin/Private/PyEditor.cpp
+++ b/Engine/Plugins/Experimental/PythonScriptPlugin/Source/PythonScriptPlugin/Private/PyEditor.cpp
@@ -227,6 +227,10 @@ void InitializeModule()
 	FPyWrapperTypeRegistry::Get().RegisterNativePythonModule(MoveTemp(NativePythonModule));
 }
 
+void ShutdownModule()
+{
+}
+
 }
 
 #endif	// WITH_EDITOR
diff --git a/Engine/Plugins/Experimental/PythonScriptPlugin/Source/PythonScriptPlugin/Private/PyEditor.h b/Engine/Plugins/Experimental/PythonScriptPlugin/Source/PythonScriptPlugin/Private/PyEditor.h
--- a/Engine/Plugins/Experimental/PythonScriptPlugin/Source/PythonScriptPlugin/Private/PyEditor.h
+++ b/Engine/Plugins/Experimental/PythonScriptPlugin/Source/PythonScriptPlugin/Private/PyEditor.h
@@ -43,6 +43,7 @@ typedef TPyPtr<FPyScopedEditorTransaction> FPyScopedEditorTransactionPtr;
 namespace PyEditor
 {
 	void InitializeModule();
+	void ShutdownModule();
 }
 
 #endif	// WITH_EDITOR
diff --git a/Engine/Plugins/Experimental/PythonScriptPlugin/Source/PythonScriptPlugin/Private/PyEngine.cpp b/Engine/Plugins/Experimental/PythonScriptPlugin/Source/PythonScriptPlugin/Private/PyEngine.cpp
--- a/Engine/Plugins/Experimental/PythonScriptPlugin/Source/PythonScriptPlugin/Private/PyEngine.cpp
+++ b/Engine/Plugins/Experimental/PythonScriptPlugin/Source/PythonScriptPlugin/Private/PyEngine.cpp
@@ -234,6 +234,10 @@ void InitializeModule()
 	FPyWrapperTypeRegistry::Get().RegisterNativePythonModule(MoveTemp(NativePythonModule));
 }
 
+void ShutdownModule()
+{
+}
+
 }
 
 #endif	// WITH_PYTHON
diff --git a/Engine/Plugins/Experimental/PythonScriptPlugin/Source/PythonScriptPlugin/Private/PyEngine.h b/Engine/Plugins/Experimental/PythonScriptPlugin/Source/PythonScriptPlugin/Private/PyEngine.h
--- a/Engine/Plugins/Experimental/PythonScriptPlugin/Source/PythonScriptPlugin/Private/PyEngine.h
+++ b/Engine/Plugins/Experimental/PythonScriptPlugin/Source/PythonScriptPlugin/Private/PyEngine.h
@@ -122,6 +122,7 @@ typedef TPyPtr<FPySelectedActorIterator> FPySelectedActorIteratorPtr;
 namespace PyEngine
 {
 	void InitializeModule();
+	void ShutdownModule();
 }
 
 #endif	// WITH_PYTHON
diff --git a/Engine/Plugins/Experimental/PythonScriptPlugin/Source/PythonScriptPlugin/Private/PySlate.cpp b/Engine/Plugins/Experimental/PythonScriptPlugin/Source/PythonScriptPlugin/Private/PySlate.cpp
--- a/Engine/Plugins/Experimental/PythonScriptPlugin/Source/PythonScriptPlugin/Private/PySlate.cpp
+++ b/Engine/Plugins/Experimental/PythonScriptPlugin/Source/PythonScriptPlugin/Private/PySlate.cpp
@@ -21,7 +21,10 @@
 namespace PySlateUtil
 {
 
-FPyDelegateHandle* RegisterSlateTickCallback(FSlateApplication::FSlateTickEvent& InSlateTickEvent, PyObject* InPyCallable)
+TArray<FDelegateHandle> PythonPreTickCallbackHandles;
+TArray<FDelegateHandle> PythonPostTickCallbackHandles;
+
+FPyDelegateHandle* RegisterSlateTickCallback(FSlateApplication::FSlateTickEvent& InSlateTickEvent, TArray<FDelegateHandle>& InOutPythonCallbackHandles, PyObject* InPyCallable)
 {
 	// We copy the PyCallable into the lambda to keep the Python object alive as long as the delegate is bound
 	FPyObjectPtr PyCallable = FPyObjectPtr::NewReference(InPyCallable);
@@ -45,10 +48,15 @@ FPyDelegateHandle* RegisterSlateTickCallback(FSlateApplication::FSlateTickEvent&
 		}
 	});
 
+	if (TickEventDelegateHandle.IsValid())
+	{
+		InOutPythonCallbackHandles.Add(TickEventDelegateHandle);
+	}
+
 	return FPyDelegateHandle::CreateInstance(TickEventDelegateHandle);
 }
 
-bool UnregisterSlateTickCallback(FSlateApplication::FSlateTickEvent& InSlateTickEvent, PyObject* InCallbackHandle)
+bool UnregisterSlateTickCallback(FSlateApplication::FSlateTickEvent& InSlateTickEvent, TArray<FDelegateHandle>& InOutPythonCallbackHandles, PyObject* InCallbackHandle)
 {
 	FPyDelegateHandlePtr PyTickEventDelegateHandle = FPyDelegateHandlePtr::StealReference(FPyDelegateHandle::CastPyObject(InCallbackHandle));
 	if (!PyTickEventDelegateHandle)
@@ -56,14 +64,25 @@ bool UnregisterSlateTickCallback(FSlateApplication::FSlateTickEvent& InSlateTick
 		return false;
 	}
 
-	if (PyTickEventDelegateHandle->Value.IsValid())
+	FDelegateHandle& TickEventDelegateHandle = PyTickEventDelegateHandle->Value;
+	if (TickEventDelegateHandle.IsValid())
 	{
-		InSlateTickEvent.Remove(PyTickEventDelegateHandle->Value);
+		InOutPythonCallbackHandles.Remove(TickEventDelegateHandle);
+		InSlateTickEvent.Remove(TickEventDelegateHandle);
 	}
 
 	return true;
 }
 
+void ClearSlateTickCallbacks(FSlateApplication::FSlateTickEvent& InSlateTickEvent, TArray<FDelegateHandle>& InOutPythonCallbackHandles)
+{
+	for (const FDelegateHandle& Handle : InOutPythonCallbackHandles)
+	{
+		InSlateTickEvent.Remove(Handle);
+	}
+	InOutPythonCallbackHandles.Empty();
+}
+
 } // namespace PySlateUtil
 
 namespace PySlate
@@ -80,7 +99,7 @@ PyObject* RegisterSlatePreTickCallback(PyObject* InSelf, PyObject* InArgs)
 
 	if (FSlateApplication::IsInitialized())
 	{
-		return (PyObject*)PySlateUtil::RegisterSlateTickCallback(FSlateApplication::Get().OnPreTick(), PyObj);
+		return (PyObject*)PySlateUtil::RegisterSlateTickCallback(FSlateApplication::Get().OnPreTick(), PySlateUtil::PythonPreTickCallbackHandles, PyObj);
 	}
 
 	return (PyObject*)FPyDelegateHandle::CreateInstance(FDelegateHandle());
@@ -95,7 +114,7 @@ PyObject* UnregisterSlatePreTickCallback(PyObject* InSelf, PyObject* InArgs)
 	}
 	check(PyObj);
 
-	if (FSlateApplication::IsInitialized() && !PySlateUtil::UnregisterSlateTickCallback(FSlateApplication::Get().OnPreTick(), PyObj))
+	if (FSlateApplication::IsInitialized() && !PySlateUtil::UnregisterSlateTickCallback(FSlateApplication::Get().OnPreTick(), PySlateUtil::PythonPreTickCallbackHandles, PyObj))
 	{
 		PyUtil::SetPythonError(PyExc_TypeError, TEXT("unregister_slate_pre_tick_callback"), *FString::Printf(TEXT("Failed to convert argument '%s' to '_DelegateHandle'"), *PyUtil::GetFriendlyTypename(PyObj)));
 		return nullptr;
@@ -115,7 +134,7 @@ PyObject* RegisterSlatePostTickCallback(PyObject* InSelf, PyObject* InArgs)
 
 	if (FSlateApplication::IsInitialized())
 	{
-		return (PyObject*)PySlateUtil::RegisterSlateTickCallback(FSlateApplication::Get().OnPostTick(), PyObj);
+		return (PyObject*)PySlateUtil::RegisterSlateTickCallback(FSlateApplication::Get().OnPostTick(), PySlateUtil::PythonPostTickCallbackHandles, PyObj);
 	}
 
 	return (PyObject*)FPyDelegateHandle::CreateInstance(FDelegateHandle());
@@ -130,7 +149,7 @@ PyObject* UnregisterSlatePostTickCallback(PyObject* InSelf, PyObject* InArgs)
 	}
 	check(PyObj);
 
-	if (FSlateApplication::IsInitialized() && !PySlateUtil::UnregisterSlateTickCallback(FSlateApplication::Get().OnPostTick(), PyObj))
+	if (FSlateApplication::IsInitialized() && !PySlateUtil::UnregisterSlateTickCallback(FSlateApplication::Get().OnPostTick(), PySlateUtil::PythonPostTickCallbackHandles, PyObj))
 	{
 		PyUtil::SetPythonError(PyExc_TypeError, TEXT("unregister_slate_post_tick_callback"), *FString::Printf(TEXT("Failed to convert argument '%s' to '_DelegateHandle'"), *PyUtil::GetFriendlyTypename(PyObj)));
 		return nullptr;
@@ -202,6 +221,17 @@ void InitializeModule()
 	FPyWrapperTypeRegistry::Get().RegisterNativePythonModule(MoveTemp(NativePythonModule));
 }
 
+void ShutdownModule()
+{
+	// We need to remove any lingering Python Slate callbacks, as they will crash if called after the interpreter has shutdown
+	if (FSlateApplication::IsInitialized())
+	{
+		FSlateApplication& SlateApplication = FSlateApplication::Get();
+		PySlateUtil::ClearSlateTickCallbacks(SlateApplication.OnPreTick(), PySlateUtil::PythonPreTickCallbackHandles);
+		PySlateUtil::ClearSlateTickCallbacks(SlateApplication.OnPostTick(), PySlateUtil::PythonPostTickCallbackHandles);
+	}
+}
+
 }
 
 #endif	// WITH_PYTHON
diff --git a/Engine/Plugins/Experimental/PythonScriptPlugin/Source/PythonScriptPlugin/Private/PySlate.h b/Engine/Plugins/Experimental/PythonScriptPlugin/Source/PythonScriptPlugin/Private/PySlate.h
--- a/Engine/Plugins/Experimental/PythonScriptPlugin/Source/PythonScriptPlugin/Private/PySlate.h
+++ b/Engine/Plugins/Experimental/PythonScriptPlugin/Source/PythonScriptPlugin/Private/PySlate.h
@@ -11,6 +11,7 @@
 namespace PySlate
 {
 	void InitializeModule();
+	void ShutdownModule();
 }
 
 #endif	// WITH_PYTHON
diff --git a/Engine/Plugins/Experimental/PythonScriptPlugin/Source/PythonScriptPlugin/Private/PythonScriptPlugin.cpp b/Engine/Plugins/Experimental/PythonScriptPlugin/Source/PythonScriptPlugin/Private/PythonScriptPlugin.cpp
--- a/Engine/Plugins/Experimental/PythonScriptPlugin/Source/PythonScriptPlugin/Private/PythonScriptPlugin.cpp
+++ b/Engine/Plugins/Experimental/PythonScriptPlugin/Source/PythonScriptPlugin/Private/PythonScriptPlugin.cpp
@@ -782,6 +782,13 @@ void FPythonScriptPlugin::ShutdownPython()
 
 	FPyReferenceCollector::Get().PurgeUnrealGeneratedTypes();
 
+#if WITH_EDITOR
+	PyEditor::ShutdownModule();
+#endif	// WITH_EDITOR
+	PyEngine::ShutdownModule();
+	PySlate::ShutdownModule();
+	PyCore::ShutdownModule();
+
 	PyUnrealModule.Reset();
 	PyDefaultGlobalDict.Reset();
 	PyDefaultLocalDict.Reset();
