Added a Blueprint node for calling Python with args

This allows people to create a Blueprint node that can define a blob of literal Python script, along with wildcard inputs and outputs to be used by the script itself.

The node takes care of marshalling data in and out of the Python script, and allows people to invoke custom Python from Blueprints without relying on Python generated BPFL types.

diff --git a/Engine/Plugins/Experimental/PythonScriptPlugin/Source/PythonScriptPlugin/Private/K2Node_ExecutePythonScript.cpp b/Engine/Plugins/Experimental/PythonScriptPlugin/Source/PythonScriptPlugin/Private/K2Node_ExecutePythonScript.cpp
new file mode 100644
--- /dev/null
+++ b/Engine/Plugins/Experimental/PythonScriptPlugin/Source/PythonScriptPlugin/Private/K2Node_ExecutePythonScript.cpp
@@ -0,0 +1,345 @@
+// Copyright 1998-2019 Epic Games, Inc. All Rights Reserved.
+
+
+#include "K2Node_ExecutePythonScript.h"
+#include "PyGenUtil.h"
+#include "PythonScriptLibrary.h"
+
+#include "UObject/Package.h"
+#include "KismetCompiler.h"
+#include "EdGraphSchema_K2.h"
+#include "K2Node_MakeArray.h"
+#include "Kismet/KismetSystemLibrary.h"
+#include "Kismet2/BlueprintEditorUtils.h"
+#include "BlueprintNodeSpawner.h"
+#include "BlueprintActionDatabaseRegistrar.h"
+
+#define LOCTEXT_NAMESPACE "K2Node_ExecutePythonScript"
+
+namespace ExecutePythonScriptUtil
+{
+
+const FName PythonScriptPinName = "PythonScript";
+const FName PythonInputsPinName = "PythonInputs";
+const FName PythonOutputsPinName = "PythonOutputs";
+
+FString PythonizePinName(const FName InPinName)
+{
+	return PyGenUtil::PythonizePropertyName(InPinName.ToString(), PyGenUtil::EPythonizeNameCase::Lower);
+}
+
+FString PythonizePinName(const UEdGraphPin* InPin)
+{
+	return PythonizePinName(InPin->PinName);
+}
+
+}
+
+UK2Node_ExecutePythonScript::UK2Node_ExecutePythonScript()
+{
+	FunctionReference.SetExternalMember(GET_FUNCTION_NAME_CHECKED(UPythonScriptLibrary, ExecutePythonScript), UPythonScriptLibrary::StaticClass());
+}
+
+void UK2Node_ExecutePythonScript::AllocateDefaultPins()
+{
+	Super::AllocateDefaultPins();
+
+	// Hide the argument name pins
+	// Note: Have to do this manually as the "HidePin" meta-data only lets you hide a single pin :(
+	UEdGraphPin* PythonInputsPin = FindPinChecked(ExecutePythonScriptUtil::PythonInputsPinName, EGPD_Input);
+	PythonInputsPin->bHidden = true;
+	UEdGraphPin* PythonOutputsPin = FindPinChecked(ExecutePythonScriptUtil::PythonOutputsPinName, EGPD_Input);
+	PythonOutputsPin->bHidden = true;
+
+	// Give the Python script pin a default value
+	UEdGraphPin* PythonScriptPin = FindPinChecked(ExecutePythonScriptUtil::PythonScriptPinName, EGPD_Input);
+	PythonScriptPin->AutogeneratedDefaultValue = TEXT("print('Hello World!')");
+	PythonScriptPin->DefaultValue = PythonScriptPin->AutogeneratedDefaultValue;
+
+	// Rename the result pin
+	UEdGraphPin* ResultPin = FindPinChecked(UEdGraphSchema_K2::PN_ReturnValue, EGPD_Output);
+	ResultPin->PinFriendlyName = LOCTEXT("ResultPinFriendlyName", "Success?");
+
+	// Add user-defined pins
+	auto AddUserDefinedPins = [this](const TArray<FName>& InPinNames, EEdGraphPinDirection InPinDirection)
+	{
+		for (const FName PinName : InPinNames)
+		{
+			UEdGraphPin* ArgPin = CreatePin(InPinDirection, UEdGraphSchema_K2::PC_Wildcard, PinName);
+			ArgPin->PinFriendlyName = FText::AsCultureInvariant(ExecutePythonScriptUtil::PythonizePinName(ArgPin));
+		}
+	};
+	AddUserDefinedPins(Inputs, EGPD_Input);
+	AddUserDefinedPins(Outputs, EGPD_Output);
+}
+
+void UK2Node_ExecutePythonScript::PostReconstructNode()
+{
+	Super::PostReconstructNode();
+
+	auto SynchronizeUserDefinedPins = [this](const TArray<FName>& InPinNames, EEdGraphPinDirection InPinDirection)
+	{
+		for (const FName PinName : InPinNames)
+		{
+			UEdGraphPin* ArgPin = FindArgumentPinChecked(PinName, InPinDirection);
+			SynchronizeArgumentPinTypeImpl(ArgPin);
+		}
+	};
+	SynchronizeUserDefinedPins(Inputs, EGPD_Input);
+	SynchronizeUserDefinedPins(Outputs, EGPD_Output);
+}
+
+void UK2Node_ExecutePythonScript::SynchronizeArgumentPinType(UEdGraphPin* Pin)
+{
+	auto SynchronizeUserDefinedPin = [this, Pin](const TArray<FName>& InPinNames, EEdGraphPinDirection InPinDirection)
+	{
+		if (Pin->Direction == InPinDirection && InPinNames.Contains(Pin->PinName))
+		{
+			// Try and find the argument pin and make sure we get the same result as the pin we were asked to update
+			// If not we may have a duplicate pin name with another non-argument pin
+			UEdGraphPin* ArgPin = FindArgumentPinChecked(Pin->PinName, InPinDirection);
+			if (ArgPin == Pin)
+			{
+				SynchronizeArgumentPinTypeImpl(Pin);
+			}
+		}
+	};
+	SynchronizeUserDefinedPin(Inputs, EGPD_Input);
+	SynchronizeUserDefinedPin(Outputs, EGPD_Output);
+}
+
+void UK2Node_ExecutePythonScript::SynchronizeArgumentPinTypeImpl(UEdGraphPin* Pin)
+{
+	FEdGraphPinType NewPinType;
+	if (Pin->LinkedTo.Num() > 0)
+	{
+		NewPinType = Pin->LinkedTo[0]->PinType;
+	}
+	else
+	{
+		NewPinType.PinCategory = UEdGraphSchema_K2::PC_Wildcard;
+	}
+
+	if (Pin->PinType != NewPinType)
+	{
+		Pin->PinType = NewPinType;
+
+		GetGraph()->NotifyGraphChanged();
+
+		UBlueprint* Blueprint = GetBlueprint();
+		if (!Blueprint->bBeingCompiled)
+		{
+			FBlueprintEditorUtils::MarkBlueprintAsModified(Blueprint);
+			Blueprint->BroadcastChanged();
+		}
+	}
+}
+
+UEdGraphPin* UK2Node_ExecutePythonScript::FindArgumentPin(const FName PinName, EEdGraphPinDirection PinDirection)
+{
+	for (int32 PinIndex = Pins.Num() - 1; PinIndex >= 0; --PinIndex)
+	{
+		UEdGraphPin* Pin = Pins[PinIndex];
+		if ((PinDirection == EGPD_MAX || PinDirection == Pin->Direction) && Pin->PinName == PinName)
+		{
+			return Pin;
+		}
+	}
+	return nullptr;
+}
+
+UEdGraphPin* UK2Node_ExecutePythonScript::FindArgumentPinChecked(const FName PinName, EEdGraphPinDirection PinDirection)
+{
+	UEdGraphPin* Pin = FindArgumentPin(PinName, PinDirection);
+	check(Pin);
+	return Pin;
+}
+
+void UK2Node_ExecutePythonScript::PostEditChangeProperty(struct FPropertyChangedEvent& PropertyChangedEvent)
+{
+	const FName PropertyName = (PropertyChangedEvent.Property  ? PropertyChangedEvent.Property->GetFName() : NAME_None);
+	if (PropertyName == GET_MEMBER_NAME_CHECKED(UK2Node_ExecutePythonScript, Inputs) || PropertyName == GET_MEMBER_NAME_CHECKED(UK2Node_ExecutePythonScript, Outputs))
+	{
+		ReconstructNode();
+		GetGraph()->NotifyGraphChanged();
+	}
+	Super::PostEditChangeProperty(PropertyChangedEvent);
+}
+
+void UK2Node_ExecutePythonScript::PinConnectionListChanged(UEdGraphPin* Pin)
+{
+	// Potentially update an argument pin type
+	SynchronizeArgumentPinType(Pin);
+}
+
+void UK2Node_ExecutePythonScript::PinTypeChanged(UEdGraphPin* Pin)
+{
+	// Potentially update an argument pin type
+	SynchronizeArgumentPinType(Pin);
+
+	Super::PinTypeChanged(Pin);
+}
+
+FText UK2Node_ExecutePythonScript::GetPinDisplayName(const UEdGraphPin* Pin) const
+{
+	if ((Inputs.Contains(Pin->PinName) || Outputs.Contains(Pin->PinName)) && !Pin->PinFriendlyName.IsEmpty())
+	{
+		return Pin->PinFriendlyName;
+	}
+	return Super::GetPinDisplayName(Pin);
+}
+
+void UK2Node_ExecutePythonScript::EarlyValidation(FCompilerResultsLog& MessageLog) const
+{
+	Super::EarlyValidation(MessageLog);
+
+	TSet<FString> AllPinNames;
+	auto ValidatePinArray = [this, &AllPinNames, &MessageLog](const TArray<FName>& InPinNames)
+	{
+		for (const FName PinName : InPinNames)
+		{
+			const FString PythonizedPinName = ExecutePythonScriptUtil::PythonizePinName(PinName);
+
+			if (PinName == UEdGraphSchema_K2::PN_Execute ||
+				PinName == UEdGraphSchema_K2::PN_Then ||
+				PinName == UEdGraphSchema_K2::PN_ReturnValue ||
+				PinName == ExecutePythonScriptUtil::PythonScriptPinName ||
+				PinName == ExecutePythonScriptUtil::PythonInputsPinName ||
+				PinName == ExecutePythonScriptUtil::PythonOutputsPinName
+				)
+			{
+				MessageLog.Error(*FText::Format(LOCTEXT("InvalidPinName_RestrictedName", "Pin '{0}' ({1}) on @@ is using a restricted name."), FText::AsCultureInvariant(PinName.ToString()), FText::AsCultureInvariant(PythonizedPinName)).ToString(), this);
+			}
+
+			if (PythonizedPinName.Len() == 0)
+			{
+				MessageLog.Error(*LOCTEXT("InvalidPinName_EmptyName", "Empty pin name found on @@").ToString(), this);
+			}
+			else
+			{
+				// Pythonized names must be a valid symbol name (alnum or _ and doesn't start with a digit)
+				const TCHAR* InvalidChar = PythonizedPinName.GetCharArray().FindByPredicate(
+					[](const TCHAR InChar)
+					{
+						return !FChar::IsAlnum(InChar) && InChar != TEXT('_') && InChar != 0;
+					});
+
+				if (InvalidChar)
+				{
+					MessageLog.Error(*FText::Format(LOCTEXT("InvalidPinName_RestrictedCharacter", "Pin '{0}' ({1}) on @@ contains '{2}' which is invalid for Python."), FText::AsCultureInvariant(PinName.ToString()), FText::AsCultureInvariant(PythonizedPinName), FText::AsCultureInvariant(FString(1, InvalidChar))).ToString(), this);
+				}
+
+				if (FChar::IsDigit(PythonizedPinName[0]))
+				{
+					MessageLog.Error(*FText::Format(LOCTEXT("InvalidPinName_LeadingDigit", "Pin '{0}' ({1}) on @@ starts with a digit which is invalid for Python."), FText::AsCultureInvariant(PinName.ToString()), FText::AsCultureInvariant(PythonizedPinName)).ToString(), this);
+				}
+
+				bool bAlreadyUsed = false;
+				AllPinNames.Add(PythonizedPinName, &bAlreadyUsed);
+
+				if (bAlreadyUsed)
+				{
+					MessageLog.Error(*FText::Format(LOCTEXT("InvalidPinName_DuplicateName", "Pin '{0}' ({1}) on @@ has the same name as another pin when exposed to Python."), FText::AsCultureInvariant(PinName.ToString()), FText::AsCultureInvariant(PythonizedPinName)).ToString(), this);
+				}
+			}
+		}
+	};
+
+	ValidatePinArray(Inputs);
+	ValidatePinArray(Outputs);
+}
+
+void UK2Node_ExecutePythonScript::ExpandNode(FKismetCompilerContext& CompilerContext, UEdGraph* SourceGraph)
+{
+	auto MakeArgumentNamesNode = [this, &CompilerContext, SourceGraph](const TArray<FName>& ArgNames, UEdGraphPin* TargetPin)
+	{
+		// Create a "Make Array" node to compile the list of arguments into an array for the Format function being called
+		UK2Node_MakeArray* MakeArrayNode = CompilerContext.SpawnIntermediateNode<UK2Node_MakeArray>(this, SourceGraph);
+		MakeArrayNode->AllocateDefaultPins();
+
+		// Connect the output of the "Make Array" pin to the destination input pin
+		// PinConnectionListChanged will set the "Make Array" node's type, only works if one pin is connected
+		UEdGraphPin* ArrayOut = MakeArrayNode->GetOutputPin();
+		ArrayOut->MakeLinkTo(TargetPin);
+		MakeArrayNode->PinConnectionListChanged(ArrayOut);
+
+		// For each argument we need to make a literal string node
+		int32 ArgIndex = 0;
+		for (const FName ArgName : ArgNames)
+		{
+			static const FName MakeLiteralStringFunctionName = GET_FUNCTION_NAME_CHECKED(UKismetSystemLibrary, MakeLiteralString);
+
+			// Create the "Make Literal String" node
+			UK2Node_CallFunction* MakeLiteralStringNode = CompilerContext.SpawnIntermediateNode<UK2Node_CallFunction>(this, SourceGraph);
+			MakeLiteralStringNode->SetFromFunction(UKismetSystemLibrary::StaticClass()->FindFunctionByName(MakeLiteralStringFunctionName));
+			MakeLiteralStringNode->AllocateDefaultPins();
+
+			// Set the literal value to the Pythonized argument name
+			UEdGraphPin* MakeLiteralIntValuePin = MakeLiteralStringNode->FindPinChecked(TEXT("Value"), EGPD_Input);
+			MakeLiteralIntValuePin->DefaultValue = ExecutePythonScriptUtil::PythonizePinName(ArgName);
+
+			// The "Make Array" node already has one pin available, so don't create one for ArgIdx == 0
+			if (ArgIndex > 0)
+			{
+				MakeArrayNode->AddInputPin();
+			}
+
+			// Find the input pin on the "Make Array" node by index and link it to the literal string
+			UEdGraphPin* ArrayIn = MakeArrayNode->FindPinChecked(FString::Printf(TEXT("[%d]"), ArgIndex++));
+			MakeLiteralStringNode->GetReturnValuePin()->MakeLinkTo(ArrayIn);
+		}
+
+		return MakeArrayNode;
+	};
+
+	// Create and connect the argument name nodes
+	MakeArgumentNamesNode(Inputs, FindPinChecked(ExecutePythonScriptUtil::PythonInputsPinName, EGPD_Input));
+	MakeArgumentNamesNode(Outputs, FindPinChecked(ExecutePythonScriptUtil::PythonOutputsPinName, EGPD_Input));
+
+	Super::ExpandNode(CompilerContext, SourceGraph);
+}
+
+bool UK2Node_ExecutePythonScript::CanPasteHere(const UEdGraph* TargetGraph) const
+{
+	bool bCanPaste = Super::CanPasteHere(TargetGraph);
+	if (bCanPaste)
+	{
+		bCanPaste &= FBlueprintEditorUtils::IsEditorUtilityBlueprint(FBlueprintEditorUtils::FindBlueprintForGraphChecked(TargetGraph));
+	}
+	return bCanPaste;
+}
+
+bool UK2Node_ExecutePythonScript::IsActionFilteredOut(const FBlueprintActionFilter& Filter)
+{
+	bool bIsFilteredOut = Super::IsActionFilteredOut(Filter);
+	if (!bIsFilteredOut)
+	{
+		for (UEdGraph* TargetGraph : Filter.Context.Graphs)
+		{
+			bIsFilteredOut |= !CanPasteHere(TargetGraph);
+		}
+	}
+	return bIsFilteredOut;
+}
+
+void UK2Node_ExecutePythonScript::GetMenuActions(FBlueprintActionDatabaseRegistrar& ActionRegistrar) const
+{
+	// actions get registered under specific object-keys; the idea is that 
+	// actions might have to be updated (or deleted) if their object-key is  
+	// mutated (or removed)... here we use the node's class (so if the node 
+	// type disappears, then the action should go with it)
+	UClass* ActionKey = GetClass();
+	// to keep from needlessly instantiating a UBlueprintNodeSpawner, first   
+	// check to make sure that the registrar is looking for actions of this type
+	// (could be regenerating actions for a specific asset, and therefore the 
+	// registrar would only accept actions corresponding to that asset)
+	if (ActionRegistrar.IsOpenForRegistration(ActionKey))
+	{
+		UBlueprintNodeSpawner* NodeSpawner = UBlueprintNodeSpawner::Create(GetClass());
+		check(NodeSpawner != nullptr);
+
+		ActionRegistrar.AddBlueprintAction(ActionKey, NodeSpawner);
+	}
+}
+
+#undef LOCTEXT_NAMESPACE
diff --git a/Engine/Plugins/Experimental/PythonScriptPlugin/Source/PythonScriptPlugin/Private/K2Node_ExecutePythonScript.h b/Engine/Plugins/Experimental/PythonScriptPlugin/Source/PythonScriptPlugin/Private/K2Node_ExecutePythonScript.h
new file mode 100644
--- /dev/null
+++ b/Engine/Plugins/Experimental/PythonScriptPlugin/Source/PythonScriptPlugin/Private/K2Node_ExecutePythonScript.h
@@ -0,0 +1,58 @@
+// Copyright 1998-2019 Epic Games, Inc. All Rights Reserved.
+
+#pragma once
+
+#include "CoreMinimal.h"
+#include "K2Node_CallFunction.h"
+#include "K2Node_ExecutePythonScript.generated.h"
+
+class UEdGraph;
+class UEdGraphPin;
+class FBlueprintActionDatabaseRegistrar;
+
+UCLASS()
+class UK2Node_ExecutePythonScript : public UK2Node_CallFunction
+{
+	GENERATED_BODY()
+	
+	UK2Node_ExecutePythonScript();
+	
+	//~ Begin UObject Interface
+	virtual void PostEditChangeProperty(struct FPropertyChangedEvent& PropertyChangedEvent) override;
+	//~ End UObject Interface
+
+	//~ Begin UEdGraphNode Interface.
+	virtual void AllocateDefaultPins() override;
+	virtual void PostReconstructNode() override;
+	virtual bool ShouldShowNodeProperties() const override { return true; }
+	virtual void PinConnectionListChanged(UEdGraphPin* Pin) override;
+	virtual void PinTypeChanged(UEdGraphPin* Pin) override;
+	virtual FText GetPinDisplayName(const UEdGraphPin* Pin) const override;
+	//~ End UEdGraphNode Interface.
+
+	//~ Begin UK2Node Interface.
+	virtual void EarlyValidation(class FCompilerResultsLog& MessageLog) const override;
+	virtual void ExpandNode(class FKismetCompilerContext& CompilerContext, UEdGraph* SourceGraph) override;
+	virtual bool CanPasteHere(const UEdGraph* TargetGraph) const override;
+	virtual bool IsActionFilteredOut(const class FBlueprintActionFilter& Filter) override;
+	virtual void GetMenuActions(FBlueprintActionDatabaseRegistrar& ActionRegistrar) const override;
+	//~ End UK2Node Interface.
+
+private:
+	/** Synchronize the type of the given argument pin with the type its connected to, or reset it to a wildcard pin if there's no connection */
+	void SynchronizeArgumentPinType(UEdGraphPin* Pin);
+	void SynchronizeArgumentPinTypeImpl(UEdGraphPin* Pin);
+
+	/** Try and find a named argument pin - this is the same as FindPin except it searches the array from the end as that's where the argument pins will be */
+	UEdGraphPin* FindArgumentPin(const FName PinName, EEdGraphPinDirection PinDirection = EGPD_MAX);
+	UEdGraphPin* FindArgumentPinChecked(const FName PinName, EEdGraphPinDirection PinDirection = EGPD_MAX);
+
+private:
+	/** User-defined input pins */
+	UPROPERTY(EditAnywhere, Category="Arguments")
+	TArray<FName> Inputs;
+
+	/** User-defined output pins */
+	UPROPERTY(EditAnywhere, Category="Arguments")
+	TArray<FName> Outputs;
+};
diff --git a/Engine/Plugins/Experimental/PythonScriptPlugin/Source/PythonScriptPlugin/Private/PythonScriptLibrary.cpp b/Engine/Plugins/Experimental/PythonScriptPlugin/Source/PythonScriptPlugin/Private/PythonScriptLibrary.cpp
--- a/Engine/Plugins/Experimental/PythonScriptPlugin/Source/PythonScriptPlugin/Private/PythonScriptLibrary.cpp
+++ b/Engine/Plugins/Experimental/PythonScriptPlugin/Source/PythonScriptPlugin/Private/PythonScriptLibrary.cpp
@@ -1,7 +1,13 @@
 // Copyright 1998-2019 Epic Games, Inc. All Rights Reserved.
 
 #include "PythonScriptLibrary.h"
-#include "IPythonScriptPlugin.h"
+#include "PythonScriptPlugin.h"
+#include "PyGIL.h"
+#include "PyUtil.h"
+#include "PyGenUtil.h"
+#include "PyConversion.h"
+#include "PyWrapperTypeRegistry.h"
+#include "UObject/Package.h"
 
 bool UPythonScriptLibrary::IsPythonAvailable()
 {
@@ -27,3 +33,124 @@ bool UPythonScriptLibrary::ExecutePythonCommandEx(const FString& PythonCommand,
 		
 	return bResult;
 }
+
+bool UPythonScriptLibrary::ExecutePythonScript(const FString& PythonScript, const TArray<FString>& PythonInputs, const TArray<FString>& PythonOutputs)
+{
+	// We should never hit this!
+	check(0);
+	return false;
+}
+
+DEFINE_FUNCTION(UPythonScriptLibrary::execExecutePythonScript)
+{
+	auto ExecuteCustomPythonScriptImpl = [&]() -> bool
+	{
+		const FString FunctionErrorName = Stack.Node->GetName();
+		const FString FunctionErrorCtxt = Stack.Node->GetOutermost()->GetName();
+
+		// Local Python context used when executing this script
+		// Has the inputs written into it prior to execution, and the outputs read from it after execution
+		FPyObjectPtr PyTempGlobalDict = FPyObjectPtr::StealReference(PyDict_Copy(FPythonScriptPlugin::Get()->GetDefaultGlobalDict()));
+		FPyObjectPtr PyTempLocalDict = PyTempGlobalDict;
+
+		// Read the standard function arguments
+		P_GET_PROPERTY_REF(UStrProperty, PythonScript);
+		P_GET_TARRAY_REF(FString, PythonInputs);
+		P_GET_TARRAY_REF(FString, PythonOutputs);
+
+		// Read the input values and write them to the Python context
+		bool bHasValidInputValues = true;
+		{
+			for (const FString& PythonInput : PythonInputs)
+			{
+				// Note: Medium term the Blueprint interpreter will change to provide us with a list of properties and 
+				// instance pointers, rather than forcing us to jump in and out of its interpreter loop (via StepCompiledIn)
+				Stack.StepCompiledIn<UProperty>(nullptr);
+				check(Stack.MostRecentProperty && Stack.MostRecentPropertyAddress);
+
+				FPyObjectPtr PyInput;
+				if (PyConversion::PythonizeProperty_Direct(Stack.MostRecentProperty, Stack.MostRecentPropertyAddress, PyInput.Get()))
+				{
+					PyDict_SetItemString(PyTempLocalDict, TCHAR_TO_UTF8(*PythonInput), PyInput);
+				}
+				else
+				{
+					PyUtil::SetPythonError(PyExc_TypeError, *FunctionErrorCtxt, *FString::Printf(TEXT("Failed to convert input property '%s' (%s) to attribute '%s' when calling function '%s' on '%s'"), *Stack.MostRecentProperty->GetName(), *Stack.MostRecentProperty->GetClass()->GetName(), *PythonInput, *FunctionErrorName, *P_THIS_OBJECT->GetName()));
+					bHasValidInputValues = false;
+				}
+			}
+		}
+
+		struct FPyOutputParam
+		{
+			UProperty* Property = nullptr;
+			uint8* PropAddr = nullptr;
+			const TCHAR* OutputName = nullptr;
+		};
+
+		// Read the output values and store them to write to later from the Python context
+		TArray<FPyOutputParam, TInlineAllocator<4>> OutParms;
+		for (const FString& PythonOutput : PythonOutputs)
+		{
+			// Note: Medium term the Blueprint interpreter will change to provide us with a list of properties and 
+			// instance pointers, rather than forcing us to jump in and out of its interpreter loop (via StepCompiledIn)
+			Stack.StepCompiledIn<UProperty>(nullptr);
+			check(Stack.MostRecentProperty && Stack.MostRecentPropertyAddress);
+
+			FPyOutputParam& OutParam = OutParms.AddDefaulted_GetRef();
+			OutParam.Property = Stack.MostRecentProperty;
+			OutParam.PropAddr = Stack.MostRecentPropertyAddress;
+			OutParam.OutputName = *PythonOutput;
+		}
+
+		P_FINISH;
+
+		// If we already failed during the parameter processing phase, then just bail now
+		if (!bHasValidInputValues)
+		{
+			return false;
+		}
+
+		// Execute the Python command
+		FPyObjectPtr PyResult = FPyObjectPtr::StealReference(FPythonScriptPlugin::Get()->EvalString(*PythonScript, TEXT("<string>"), Py_file_input, PyTempGlobalDict, PyTempGlobalDict));
+		FPyWrapperTypeReinstancer::Get().ProcessPending();
+
+		// Read the output values from the Python context
+		if (PyResult)
+		{
+			for (const FPyOutputParam& OutParam : OutParms)
+			{
+				PyObject* PyOutput = PyDict_GetItemString(PyTempLocalDict, TCHAR_TO_UTF8(OutParam.OutputName));
+				if (!PyOutput)
+				{
+					PyUtil::SetPythonError(PyExc_TypeError, *FunctionErrorCtxt, *FString::Printf(TEXT("Failed to find attribute '%s' for output property '%s' (%s) when calling function '%s' on '%s'"), OutParam.OutputName, *OutParam.Property->GetName(), *OutParam.Property->GetClass()->GetName() , *FunctionErrorName, *P_THIS_OBJECT->GetName()));
+					return false;
+				}
+
+				if (!PyConversion::NativizeProperty_Direct(PyOutput, OutParam.Property, OutParam.PropAddr))
+				{
+					PyUtil::SetPythonError(PyExc_TypeError, *FunctionErrorCtxt, *FString::Printf(TEXT("Failed to convert output property '%s' (%s) from attribute '%s' when calling function '%s' on '%s'"), *OutParam.Property->GetName(), *OutParam.Property->GetClass()->GetName(), OutParam.OutputName, *FunctionErrorName, *P_THIS_OBJECT->GetName()));
+					return false;
+				}
+			}
+
+			return true;
+		}
+
+		return false;
+	};
+
+	// Execute Python code within this block
+	{
+		FPyScopedGIL GIL;
+		if (ExecuteCustomPythonScriptImpl())
+		{
+			*(bool*)RESULT_PARAM = true;
+		}
+		else
+		{
+			PyUtil::ReThrowPythonError();
+			*(bool*)RESULT_PARAM = false;
+		}
+	}
+}
diff --git a/Engine/Plugins/Experimental/PythonScriptPlugin/Source/PythonScriptPlugin/Private/PythonScriptLibrary.h b/Engine/Plugins/Experimental/PythonScriptPlugin/Source/PythonScriptPlugin/Private/PythonScriptLibrary.h
--- a/Engine/Plugins/Experimental/PythonScriptPlugin/Source/PythonScriptPlugin/Private/PythonScriptLibrary.h
+++ b/Engine/Plugins/Experimental/PythonScriptPlugin/Source/PythonScriptPlugin/Private/PythonScriptLibrary.h
@@ -38,4 +38,15 @@ public:
 	 */
 	UFUNCTION(BlueprintCallable, Category = "Python|Execution", meta=(DisplayName="Execute Python Command (Advanced)", AdvancedDisplay="ExecutionMode,FileExecutionScope"))
 	static bool ExecutePythonCommandEx(UPARAM(meta=(MultiLine=True)) const FString& PythonCommand, FString& CommandResult, TArray<FPythonLogOutputEntry>& LogOutput, const EPythonCommandExecutionMode ExecutionMode = EPythonCommandExecutionMode::ExecuteFile, const EPythonFileExecutionScope FileExecutionScope = EPythonFileExecutionScope::Private);
+
+	/**
+	 * Execute a Python script with argument marshaling.
+	 * @param PythonScript This literal Python code to run.
+	 * @param PythonInputs The variadic input argument names (internal; set by UK2Node_ExecutePythonScript).
+	 * @param PythonInputs The variadic output argument names (internal; set by UK2Node_ExecutePythonScript).
+	 * @return true if the script ran successfully, false if there were errors.
+	 */
+    UFUNCTION(BlueprintCallable, CustomThunk, Category = "Python|Execution", meta=(Variadic, BlueprintInternalUseOnly="true"))
+    static bool ExecutePythonScript(UPARAM(meta=(MultiLine=True)) const FString& PythonScript, const TArray<FString>& PythonInputs, const TArray<FString>& PythonOutputs);
+	DECLARE_FUNCTION(execExecutePythonScript);
 };
diff --git a/Engine/Plugins/Experimental/PythonScriptPlugin/Source/PythonScriptPlugin/Private/PythonScriptPlugin.h b/Engine/Plugins/Experimental/PythonScriptPlugin/Source/PythonScriptPlugin/Private/PythonScriptPlugin.h
--- a/Engine/Plugins/Experimental/PythonScriptPlugin/Source/PythonScriptPlugin/Private/PythonScriptPlugin.h
+++ b/Engine/Plugins/Experimental/PythonScriptPlugin/Source/PythonScriptPlugin/Private/PythonScriptPlugin.h
@@ -143,6 +143,11 @@ public:
 
 	/** Run a Python file */
 	bool RunFile(const TCHAR* InFile, const TCHAR* InArgs, FPythonCommandEx& InOutPythonCommand);
+
+	PyObject* GetDefaultGlobalDict() { return PyDefaultGlobalDict.Get(); }
+	PyObject* GetDefaultLocalDict()  { return PyDefaultLocalDict.Get();  }
+	PyObject* GetConsoleGlobalDict() { return PyConsoleGlobalDict.Get(); }
+	PyObject* GetConsoleLocalDict()  { return PyConsoleLocalDict.Get();  }
 #endif	// WITH_PYTHON
 
 private:
diff --git a/Engine/Plugins/Experimental/PythonScriptPlugin/Source/PythonScriptPlugin/PythonScriptPlugin.Build.cs b/Engine/Plugins/Experimental/PythonScriptPlugin/Source/PythonScriptPlugin/PythonScriptPlugin.Build.cs
--- a/Engine/Plugins/Experimental/PythonScriptPlugin/Source/PythonScriptPlugin/PythonScriptPlugin.Build.cs
+++ b/Engine/Plugins/Experimental/PythonScriptPlugin/Source/PythonScriptPlugin/PythonScriptPlugin.Build.cs
@@ -46,6 +46,7 @@ namespace UnrealBuildTool.Rules
 						"UnrealEd",
 						"EditorSubsystem",
 						"BlueprintGraph",
+						"KismetCompiler",
 					}
 				);
 			}
