Improved Python error reporting for invalid ScriptName and ScriptMethod values

Bad values here can cause the unreal.py file to fail to parse, so we now validate that user-defined values are valid for Python when generating the glue and report an error for any that are invalid.

This also re-adds the type prefixes to any class, struct, or enum types that would start with a number without their prefix. Once again, this is needed to avoid invalid symbol names in the unreal.py file.

diff --git a/Engine/Plugins/Experimental/PythonScriptPlugin/Source/PythonScriptPlugin/Private/K2Node_ExecutePythonScript.cpp b/Engine/Plugins/Experimental/PythonScriptPlugin/Source/PythonScriptPlugin/Private/K2Node_ExecutePythonScript.cpp
--- a/Engine/Plugins/Experimental/PythonScriptPlugin/Source/PythonScriptPlugin/Private/K2Node_ExecutePythonScript.cpp
+++ b/Engine/Plugins/Experimental/PythonScriptPlugin/Source/PythonScriptPlugin/Private/K2Node_ExecutePythonScript.cpp
@@ -217,21 +217,10 @@ void UK2Node_ExecutePythonScript::EarlyValidation(FCompilerResultsLog& MessageLo
 			}
 			else
 			{
-				// Pythonized names must be a valid symbol name (alnum or _ and doesn't start with a digit)
-				const TCHAR* InvalidChar = PythonizedPinName.GetCharArray().FindByPredicate(
-					[](const TCHAR InChar)
-					{
-						return !FChar::IsAlnum(InChar) && InChar != TEXT('_') && InChar != 0;
-					});
-
-				if (InvalidChar)
+				FText NameValidationError;
+				if (!PyGenUtil::IsValidName(PythonizedPinName, &NameValidationError))
 				{
-					MessageLog.Error(*FText::Format(LOCTEXT("InvalidPinName_RestrictedCharacter", "Pin '{0}' ({1}) on @@ contains '{2}' which is invalid for Python."), FText::AsCultureInvariant(PinName.ToString()), FText::AsCultureInvariant(PythonizedPinName), FText::AsCultureInvariant(FString(1, InvalidChar))).ToString(), this);
-				}
-
-				if (FChar::IsDigit(PythonizedPinName[0]))
-				{
-					MessageLog.Error(*FText::Format(LOCTEXT("InvalidPinName_LeadingDigit", "Pin '{0}' ({1}) on @@ starts with a digit which is invalid for Python."), FText::AsCultureInvariant(PinName.ToString()), FText::AsCultureInvariant(PythonizedPinName)).ToString(), this);
+					MessageLog.Error(*FText::Format(LOCTEXT("InvalidPinName_GenValidationError", "Pin '{0}' ({1}) on @@ failed name validation: {2}."), FText::AsCultureInvariant(PinName.ToString()), FText::AsCultureInvariant(PythonizedPinName), NameValidationError).ToString(), this);
 				}
 
 				bool bAlreadyUsed = false;
diff --git a/Engine/Plugins/Experimental/PythonScriptPlugin/Source/PythonScriptPlugin/Private/PyGenUtil.cpp b/Engine/Plugins/Experimental/PythonScriptPlugin/Source/PythonScriptPlugin/Private/PyGenUtil.cpp
--- a/Engine/Plugins/Experimental/PythonScriptPlugin/Source/PythonScriptPlugin/Private/PyGenUtil.cpp
+++ b/Engine/Plugins/Experimental/PythonScriptPlugin/Source/PythonScriptPlugin/Private/PyGenUtil.cpp
@@ -24,6 +24,7 @@
 #include "Engine/BlueprintGeneratedClass.h"
 #include "Engine/UserDefinedStruct.h"
 #include "Engine/UserDefinedEnum.h"
+#include "GameFramework/Actor.h"
 #include "Interfaces/IPluginManager.h"
 
 #if WITH_PYTHON
@@ -1565,6 +1566,46 @@ bool ShouldExportFunction(const UFunction* InFunc)
 	return IsScriptExposedFunction(InFunc);
 }
 
+bool IsValidName(const FString& InName, FText* OutError)
+{
+	if (InName.Len() == 0)
+	{
+		if (OutError)
+		{
+			*OutError = NSLOCTEXT("PyGenUtil", "InvalidName_EmptyName", "Name is empty");
+		}
+		return false;
+	}
+
+	// Names must be a valid symbol (alnum or _ and doesn't start with a digit)
+
+	const TCHAR* InvalidChar = InName.GetCharArray().FindByPredicate(
+		[](const TCHAR InChar)
+		{
+			return !FChar::IsAlnum(InChar) && InChar != TEXT('_') && InChar != 0;
+		});
+
+	if (InvalidChar)
+	{
+		if (OutError)
+		{
+			*OutError = FText::Format(NSLOCTEXT("PyGenUtil", "InvalidName_RestrictedCharacter", "Name contains '{0}' which is invalid for Python"), FText::AsCultureInvariant(FString(1, InvalidChar)));
+		}
+		return false;
+	}
+
+	if (FChar::IsDigit(InName[0]))
+	{
+		if (OutError)
+		{
+			*OutError = NSLOCTEXT("PyGenUtil", "InvalidName_LeadingDigit", "Name starts with a digit which is invalid for Python");
+		}
+		return false;
+	}
+
+	return true;
+}
+
 FString PythonizeName(const FString& InName, const EPythonizeNameCase InNameCase)
 {
 	static const TSet<FString, FCaseSensitiveStringSetFuncs> ReservedKeywords = {
@@ -2440,6 +2481,12 @@ bool GetFieldPythonNameFromMetaDataImpl(const UField* InField, const FName InMet
 				OutFieldName.RemoveAt(SemiColonIndex, OutFieldName.Len() - SemiColonIndex, /*bAllowShrinking*/false);
 			}
 
+			FText ValidationError;
+			if (!IsValidName(OutFieldName, &ValidationError))
+			{
+				REPORT_PYTHON_GENERATION_ISSUE(Error, TEXT("'%s' has an invalid '%s' meta-data value '%s': %s."), *InField->GetPathName(), *InMetaDataKey.ToString(), *OutFieldName, *ValidationError.ToString());
+			}
+
 			return true;
 		}
 	}
@@ -2472,6 +2519,15 @@ bool GetDeprecatedFieldPythonNamesFromMetaDataImpl(const UField* InField, const
 				return InStr.IsEmpty();
 			});
 
+			for (const FString& FieldNamePart : OutFieldNames)
+			{
+				FText ValidationError;
+				if (!IsValidName(FieldNamePart, &ValidationError))
+				{
+					REPORT_PYTHON_GENERATION_ISSUE(Error, TEXT("'%s' has an invalid '%s' meta-data value '%s' (from '%s'): %s."), *InField->GetPathName(), *InMetaDataKey.ToString(), *FieldNamePart, *FieldName, *ValidationError.ToString());
+				}
+			}
+
 			return true;
 		}
 	}
@@ -2499,6 +2555,25 @@ FString GetFieldPythonNameImpl(const UField* InField, const FName InMetaDataKey)
 		{
 			FieldName.RemoveAt(0, 1, /*bAllowShrinking*/false);
 		}
+
+		// Classes, structs, and enums will no longer have their C++ prefix at this point
+		// These prefixes can prevent types that start with numbers from being a compile error in C++
+		// Ideally we don't want those prefixes to be used for Python, but we must ensure a valid symbol name
+		if (FieldName.Len() >= 0 && FChar::IsDigit(FieldName[0]))
+		{
+			if (const UClass* Class = Cast<UClass>(InField))
+			{
+				FieldName.InsertAt(0, Class->IsChildOf<AActor>() ? TEXT('A') : TEXT('U'));
+			}
+			else if (InField->IsA<UScriptStruct>())
+			{
+				FieldName.InsertAt(0, TEXT('F'));
+			}
+			else if (InField->IsA<UEnum>())
+			{
+				FieldName.InsertAt(0, TEXT('E'));
+			}
+		}
 	}
 
 	return FieldName;
@@ -2618,6 +2693,12 @@ FString GetEnumEntryPythonName(const UEnum* InEnum, const int32 InEntryIndex)
 			{
 				EnumEntryName.RemoveAt(SemiColonIndex, EnumEntryName.Len() - SemiColonIndex, /*bAllowShrinking*/false);
 			}
+
+			FText ValidationError;
+			if (!IsValidName(EnumEntryName, &ValidationError))
+			{
+				REPORT_PYTHON_GENERATION_ISSUE(Error, TEXT("'%s' has an invalid 'ScriptName' meta-data value '%s': %s."), *InEnum->GetPathName(), *EnumEntryName, *ValidationError.ToString());
+			}
 		}
 	}
 	
diff --git a/Engine/Plugins/Experimental/PythonScriptPlugin/Source/PythonScriptPlugin/Private/PyGenUtil.h b/Engine/Plugins/Experimental/PythonScriptPlugin/Source/PythonScriptPlugin/Private/PyGenUtil.h
--- a/Engine/Plugins/Experimental/PythonScriptPlugin/Source/PythonScriptPlugin/Private/PyGenUtil.h
+++ b/Engine/Plugins/Experimental/PythonScriptPlugin/Source/PythonScriptPlugin/Private/PyGenUtil.h
@@ -981,6 +981,9 @@ namespace PyGenUtil
 	/** Should the given function be exported to Python? */
 	bool ShouldExportFunction(const UFunction* InFunc);
 
+	/** Check that the given name will be valid for Python () */
+	bool IsValidName(const FString& InName, FText* OutError = nullptr);
+
 	/** Given a CamelCase name, convert it to snake_case */
 	FString PythonizeName(const FString& InName, const EPythonizeNameCase InNameCase);
 
