Fix CIS static analysis warning (unreachable code)

#rnx

diff --git a/Engine/Plugins/Experimental/PythonScriptPlugin/Source/PythonScriptPlugin/Private/PyGenUtil.cpp b/Engine/Plugins/Experimental/PythonScriptPlugin/Source/PythonScriptPlugin/Private/PyGenUtil.cpp
--- a/Engine/Plugins/Experimental/PythonScriptPlugin/Source/PythonScriptPlugin/Private/PyGenUtil.cpp
+++ b/Engine/Plugins/Experimental/PythonScriptPlugin/Source/PythonScriptPlugin/Private/PyGenUtil.cpp
@@ -2404,16 +2404,9 @@ FString GetFieldModule(const UField* InField)
 	{
 		return PackageName.RightChop(8); // Chop "/Script/" from the name
 	}
-	else
-	{
-		// Not a native module!
-		return FString();
-	}
-
-	check(PackageName[0] == TEXT('/'));
-	int32 RootNameEnd = 1;
-	for (; PackageName[RootNameEnd] != TEXT('/'); ++RootNameEnd) {}
-	return PackageName.Mid(1, RootNameEnd - 1);
+	
+	// Not a native module!
+	return FString();
 }
 
 FString GetFieldPlugin(const UField* InField)
