Emit RuntimeError in Python for Blueprint execution errors

This previously used the generic Exception type

diff --git a/Engine/Plugins/Experimental/PythonScriptPlugin/Source/PythonScriptPlugin/Private/PyUtil.cpp b/Engine/Plugins/Experimental/PythonScriptPlugin/Source/PythonScriptPlugin/Private/PyUtil.cpp
--- a/Engine/Plugins/Experimental/PythonScriptPlugin/Source/PythonScriptPlugin/Private/PyUtil.cpp
+++ b/Engine/Plugins/Experimental/PythonScriptPlugin/Source/PythonScriptPlugin/Private/PyUtil.cpp
@@ -536,7 +536,7 @@ bool InvokeFunctionCall(UObject* InObj, const UFunction* InFunc, void* InBasePar
 	{
 		if (Verbosity == ELogVerbosity::Error)
 		{
-			SetPythonError(PyExc_Exception, InErrorCtxt, ExceptionMessage);
+			SetPythonError(PyExc_RuntimeError, InErrorCtxt, ExceptionMessage);
 			bThrewException = true;
 		}
 		else if (Verbosity == ELogVerbosity::Warning)
