Fixing Python Script Plugin crashes caused by FProperty changes

diff --git a/Engine/Plugins/Experimental/PythonScriptPlugin/Source/PythonScriptPlugin/Private/PyWrapperArray.cpp b/Engine/Plugins/Experimental/PythonScriptPlugin/Source/PythonScriptPlugin/Private/PyWrapperArray.cpp
--- a/Engine/Plugins/Experimental/PythonScriptPlugin/Source/PythonScriptPlugin/Private/PyWrapperArray.cpp
+++ b/Engine/Plugins/Experimental/PythonScriptPlugin/Source/PythonScriptPlugin/Private/PyWrapperArray.cpp
@@ -1535,7 +1535,7 @@ void FPyWrapperArrayMetaData::AddReferencedObjects(FPyWrapperBase* Instance, FRe
 	FPyWrapperArray* Self = static_cast<FPyWrapperArray*>(Instance);
 	if (Self->ArrayProp && Self->ArrayInstance && !Self->OwnerContext.HasOwner())
 	{
-		Collector.AddReferencedObject(Self->ArrayProp);
+		const_cast<FArrayProperty*>(Self->ArrayProp)->AddReferencedObjects(Collector);
 		FPyReferenceCollector::AddReferencedObjectsFromProperty(Collector, Self->ArrayProp, Self->ArrayInstance);
 	}
 }
diff --git a/Engine/Plugins/Experimental/PythonScriptPlugin/Source/PythonScriptPlugin/Private/PyWrapperObject.cpp b/Engine/Plugins/Experimental/PythonScriptPlugin/Source/PythonScriptPlugin/Private/PyWrapperObject.cpp
--- a/Engine/Plugins/Experimental/PythonScriptPlugin/Source/PythonScriptPlugin/Private/PyWrapperObject.cpp
+++ b/Engine/Plugins/Experimental/PythonScriptPlugin/Source/PythonScriptPlugin/Private/PyWrapperObject.cpp
@@ -238,7 +238,7 @@ PyObject* FPyWrapperObject::CallFunction_Impl(UObject* InObj, const PyGenUtil::F
 			}
 		}
 
-		if (InFuncDef.Func->Children == nullptr)
+		if (InFuncDef.Func->ChildProperties == nullptr)
 		{
 			// No return value
 			if (!PyUtil::InvokeFunctionCall(InObj, InFuncDef.Func, nullptr, InErrorCtxt))
