Fixing logic behind exposing Unreal classes to Python

diff --git a/Engine/Plugins/Experimental/PythonScriptPlugin/Source/PythonScriptPlugin/Private/PyGenUtil.cpp b/Engine/Plugins/Experimental/PythonScriptPlugin/Source/PythonScriptPlugin/Private/PyGenUtil.cpp
--- a/Engine/Plugins/Experimental/PythonScriptPlugin/Source/PythonScriptPlugin/Private/PyGenUtil.cpp
+++ b/Engine/Plugins/Experimental/PythonScriptPlugin/Source/PythonScriptPlugin/Private/PyGenUtil.cpp
@@ -1374,28 +1374,27 @@ bool IsScriptExposedFunction(const UFunction* InFunc)
 		&& !InFunc->HasMetaData(NativeMakeFuncMetaDataKey);
 }
 
-bool IsScriptExposedField(const UField* InField)
+bool HasScriptExposedFields(const UStruct* InStruct)
 {
-	if (const FProperty* Prop = CastField<const FProperty>(InField))
-	{
-		return IsScriptExposedProperty(Prop);
-	}
-
-	if (const UFunction* Func = Cast<const UFunction>(InField))
+	for (TFieldIterator<const UField> FieldIt(InStruct); FieldIt; ++FieldIt)
 	{
-		return IsScriptExposedFunction(Func);
+		if (const UFunction* Func = Cast<const UFunction>(*FieldIt))
+		{
+			if (IsScriptExposedFunction(Func))
+			{
+				return true;
+			}
+		}
 	}
 
-	return false;
-}
-
-bool HasScriptExposedFields(const UStruct* InStruct)
-{
-	for (TFieldIterator<const UField> FieldIt(InStruct); FieldIt; ++FieldIt)
+	for (TFieldIterator<const FField> FieldIt(InStruct); FieldIt; ++FieldIt)
 	{
-		if (IsScriptExposedField(*FieldIt))
+		if (const FProperty* Prop = CastField<const FProperty>(*FieldIt))
 		{
-			return true;
+			if (IsScriptExposedProperty(Prop))
+			{
+				return true;
+			}
 		}
 	}
 
diff --git a/Engine/Plugins/Experimental/PythonScriptPlugin/Source/PythonScriptPlugin/Private/PyWrapperTypeRegistry.cpp b/Engine/Plugins/Experimental/PythonScriptPlugin/Source/PythonScriptPlugin/Private/PyWrapperTypeRegistry.cpp
--- a/Engine/Plugins/Experimental/PythonScriptPlugin/Source/PythonScriptPlugin/Private/PyWrapperTypeRegistry.cpp
+++ b/Engine/Plugins/Experimental/PythonScriptPlugin/Source/PythonScriptPlugin/Private/PyWrapperTypeRegistry.cpp
@@ -1151,14 +1151,16 @@ PyTypeObject* FPyWrapperTypeRegistry::GenerateWrappedClassType(const UClass* InC
 	const FString PythonClassName = PyGenUtil::GetClassPythonName(InClass);
 	GeneratedWrappedType->TypeName = PyGenUtil::TCHARToUTF8Buffer(*PythonClassName);
 
-	for (TFieldIterator<const UField> FieldIt(InClass, EFieldIteratorFlags::ExcludeSuper); FieldIt; ++FieldIt)
+	for (TFieldIterator<const FField> FieldIt(InClass, EFieldIteratorFlags::ExcludeSuper); FieldIt; ++FieldIt)
 	{
 		if (const FProperty* Prop = CastField<const FProperty>(*FieldIt))
 		{
 			GenerateWrappedProperty(Prop);
 			continue;
 		}
-
+	}
+	for (TFieldIterator<const UField> FieldIt(InClass, EFieldIteratorFlags::ExcludeSuper); FieldIt; ++FieldIt)
+	{
 		if (const UFunction* Func = Cast<const UFunction>(*FieldIt))
 		{
 			GenerateWrappedMethod(Func);
