Fixing static analysis warnings

diff --git a/Engine/Plugins/Experimental/PythonScriptPlugin/Source/PythonScriptPlugin/Private/PyReferenceCollector.cpp b/Engine/Plugins/Experimental/PythonScriptPlugin/Source/PythonScriptPlugin/Private/PyReferenceCollector.cpp
--- a/Engine/Plugins/Experimental/PythonScriptPlugin/Source/PythonScriptPlugin/Private/PyReferenceCollector.cpp
+++ b/Engine/Plugins/Experimental/PythonScriptPlugin/Source/PythonScriptPlugin/Private/PyReferenceCollector.cpp
@@ -370,12 +370,9 @@ void FPyReferenceCollector::AddReferencedObjectsFromPropertyInternal(FReferenceC
 			{
 				if (ScriptMapHelper.IsValidIndex(SparseElementIndex))
 				{
-					if (ScriptMapHelper.IsValidIndex(SparseElementIndex))
-					{
-						// Note: We use the pair pointer below as AddReferencedObjectsFromPropertyInternal expects a base address and the key/value property will apply the correct offset from the base
-						AddReferencedObjectsFromPropertyInternal(InCollector, ScriptMapHelper.GetKeyProperty(), ScriptMapHelper.GetPairPtr(SparseElementIndex), InFlags, bMapKeysChanged);
-						AddReferencedObjectsFromPropertyInternal(InCollector, ScriptMapHelper.GetValueProperty(), ScriptMapHelper.GetPairPtr(SparseElementIndex), InFlags, bMapValuesChanged);
-					}
+					// Note: We use the pair pointer below as AddReferencedObjectsFromPropertyInternal expects a base address and the key/value property will apply the correct offset from the base
+					AddReferencedObjectsFromPropertyInternal(InCollector, ScriptMapHelper.GetKeyProperty(), ScriptMapHelper.GetPairPtr(SparseElementIndex), InFlags, bMapKeysChanged);
+					AddReferencedObjectsFromPropertyInternal(InCollector, ScriptMapHelper.GetValueProperty(), ScriptMapHelper.GetPairPtr(SparseElementIndex), InFlags, bMapValuesChanged);
 				}
 			}
 
diff --git a/Engine/Source/Runtime/CoreUObject/Private/UObject/EnumProperty.cpp b/Engine/Source/Runtime/CoreUObject/Private/UObject/EnumProperty.cpp
--- a/Engine/Source/Runtime/CoreUObject/Private/UObject/EnumProperty.cpp
+++ b/Engine/Source/Runtime/CoreUObject/Private/UObject/EnumProperty.cpp
@@ -102,7 +102,7 @@ void FEnumProperty::AddCppProperty(FProperty* Inner)
 	check(!UnderlyingProp);
 	UnderlyingProp = CastFieldChecked<FNumericProperty>(Inner);
 	check(UnderlyingProp->GetOwner<FEnumProperty>() == this);
-	if (UnderlyingProp && UnderlyingProp->HasAnyPropertyFlags(CPF_HasGetValueTypeHash))
+	if (UnderlyingProp->HasAnyPropertyFlags(CPF_HasGetValueTypeHash))
 	{
 		PropertyFlags |= CPF_HasGetValueTypeHash;
 	}
diff --git a/Engine/Source/Runtime/CoreUObject/Public/UObject/PropertyIterator.h b/Engine/Source/Runtime/CoreUObject/Public/UObject/PropertyIterator.h
--- a/Engine/Source/Runtime/CoreUObject/Public/UObject/PropertyIterator.h
+++ b/Engine/Source/Runtime/CoreUObject/Public/UObject/PropertyIterator.h
@@ -1,4 +1,4 @@
-// Copyright 1998-2017 Epic Games, Inc. All Rights Reserved.
+// Copyright 1998-2019 Epic Games, Inc. All Rights Reserved.
 
 /*=============================================================================
 PropertyIterrator.h: TPropertyIterator implementation.
