Improved error message when creating a Python property for a container type

It is a common mistake to omit the type(s) for the container elements, so this case will now present an error stating the reason and the correct fix.

diff --git a/Engine/Plugins/Experimental/PythonScriptPlugin/Source/PythonScriptPlugin/Private/PyUtil.cpp b/Engine/Plugins/Experimental/PythonScriptPlugin/Source/PythonScriptPlugin/Private/PyUtil.cpp
--- a/Engine/Plugins/Experimental/PythonScriptPlugin/Source/PythonScriptPlugin/Private/PyUtil.cpp
+++ b/Engine/Plugins/Experimental/PythonScriptPlugin/Source/PythonScriptPlugin/Private/PyUtil.cpp
@@ -335,6 +335,31 @@ bool CalculatePropertyDef(PyTypeObject* InPyType, FPropertyDef& OutPropertyDef)
 
 bool CalculatePropertyDef(PyObject* InPyObj, FPropertyDef& OutPropertyDef)
 {
+	// It is a common error for a user to pass the container type directly
+	// rather than an instance of it that defines the sub-types
+	// eg) To pass "unreal.Map" rather than "unreal.Map(int, str)"
+	// This tests for that case and emits a suitable error
+	if (PyType_Check(InPyObj))
+	{
+		if (PyObject_IsSubclass(InPyObj, (PyObject*)&PyWrapperArrayType) == 1)
+		{
+			SetPythonError(PyExc_TypeError, InPyObj, TEXT("Cannot create a property definition from 'Array' directly! It must be an instance specifying the element type, eg) 'Array(int)'."));
+			return false;
+		}
+
+		if (PyObject_IsSubclass(InPyObj, (PyObject*)&PyWrapperSetType) == 1)
+		{
+			SetPythonError(PyExc_TypeError, InPyObj, TEXT("Cannot create a property definition from 'Set' directly! It must be an instance specifying the element type, eg) 'Set(int)'."));
+			return false;
+		}
+
+		if (PyObject_IsSubclass(InPyObj, (PyObject*)&PyWrapperMapType) == 1)
+		{
+			SetPythonError(PyExc_TypeError, InPyObj, TEXT("Cannot create a property definition from 'Map' directly! It must be an instance specifying the key and value types, eg) 'Map(int, str)'."));
+			return false;
+		}
+	}
+
 	if (PyObject_IsInstance(InPyObj, (PyObject*)&PyWrapperArrayType) == 1)
 	{
 		FPyWrapperArray* PyArray = (FPyWrapperArray*)InPyObj;
