Improved error message when attempting to create a container property in Python that nests another container

This isn't supported by the reflection system and is already disallowed, but the error message was confusing. It now gives an error message more like UHT, and suggests using a struct as an alternative to direct nesting.

diff --git a/Engine/Plugins/Experimental/PythonScriptPlugin/Source/PythonScriptPlugin/Private/PyUtil.cpp b/Engine/Plugins/Experimental/PythonScriptPlugin/Source/PythonScriptPlugin/Private/PyUtil.cpp
--- a/Engine/Plugins/Experimental/PythonScriptPlugin/Source/PythonScriptPlugin/Private/PyUtil.cpp
+++ b/Engine/Plugins/Experimental/PythonScriptPlugin/Source/PythonScriptPlugin/Private/PyUtil.cpp
@@ -642,7 +642,16 @@ bool InspectFunctionArgs(PyObject* InFunc, TArray<FString>& OutArgNames, TArray<
 
 int ValidateContainerTypeParam(PyObject* InPyObj, FPropertyDef& OutPropDef, const char* InPythonArgName, const TCHAR* InErrorCtxt)
 {
-	if (PyObject_IsInstance(InPyObj, (PyObject*)&PyType_Type) != 1)
+	if (PyObject_IsInstance(InPyObj, (PyObject*)&PyWrapperArrayType) == 1 ||
+		PyObject_IsInstance(InPyObj, (PyObject*)&PyWrapperSetType) == 1 ||
+		PyObject_IsInstance(InPyObj, (PyObject*)&PyWrapperMapType) == 1
+		)
+	{
+		SetPythonError(PyExc_TypeError, InErrorCtxt, *FString::Printf(TEXT("'%s' (%s) cannot be a container element type (directly nested containers are not supported - consider using an intermediary struct instead)"), UTF8_TO_TCHAR(InPythonArgName), *GetFriendlyTypename(InPyObj)));
+		return -1;
+	}
+
+	if (PyType_Check(InPyObj) != 1)
 	{
 		SetPythonError(PyExc_TypeError, InErrorCtxt, *FString::Printf(TEXT("'%s' (%s) must be a type"), UTF8_TO_TCHAR(InPythonArgName), *GetFriendlyTypename(InPyObj)));
 		return -1;
@@ -650,19 +659,19 @@ int ValidateContainerTypeParam(PyObject* InPyObj, FPropertyDef& OutPropDef, cons
 
 	if (!CalculatePropertyDef((PyTypeObject*)InPyObj, OutPropDef))
 	{
-		SetPythonError(PyExc_TypeError, InErrorCtxt, *FString::Printf(TEXT("Failed to convert '%s' (%s) to a 'FProperty' class"), UTF8_TO_TCHAR(InPythonArgName), *GetFriendlyTypename(InPyObj)));
+		SetPythonError(PyExc_TypeError, InErrorCtxt, *FString::Printf(TEXT("Failed to convert '%s' (%s) to a 'Property' class"), UTF8_TO_TCHAR(InPythonArgName), *GetFriendlyTypename(InPyObj)));
 		return -1;
 	}
 
 	if (OutPropDef.KeyDef.IsValid() || OutPropDef.ValueDef.IsValid())
 	{
-		SetPythonError(PyExc_TypeError, InErrorCtxt, *FString::Printf(TEXT("'%s' (%s) cannot be a container type"), UTF8_TO_TCHAR(InPythonArgName), *GetFriendlyTypename(InPyObj)));
+		SetPythonError(PyExc_TypeError, InErrorCtxt, *FString::Printf(TEXT("'%s' (%s) cannot be a container element type"), UTF8_TO_TCHAR(InPythonArgName), *GetFriendlyTypename(InPyObj)));
 		return -1;
 	}
 
 	if (OutPropDef.PropertyClass->HasAnyClassFlags(CLASS_Abstract))
 	{
-		SetPythonError(PyExc_TypeError, InErrorCtxt, *FString::Printf(TEXT("'%s' (%s) converted to '%s' which is an abstract 'FProperty' class"), UTF8_TO_TCHAR(InPythonArgName), *GetFriendlyTypename(InPyObj), *OutPropDef.PropertyClass->GetName()));
+		SetPythonError(PyExc_TypeError, InErrorCtxt, *FString::Printf(TEXT("'%s' (%s) converted to '%s' which is an abstract 'Property' class"), UTF8_TO_TCHAR(InPythonArgName), *GetFriendlyTypename(InPyObj), *OutPropDef.PropertyClass->GetName()));
 		return -1;
 	}
 
