Python: Added an editor preferences section for Python which contains a way to enable Developer Mode locally without affecting other users

diff --git a/Engine/Plugins/Experimental/PythonScriptPlugin/Source/PythonScriptPlugin/Private/PythonScriptPlugin.cpp b/Engine/Plugins/Experimental/PythonScriptPlugin/Source/PythonScriptPlugin/Private/PythonScriptPlugin.cpp
--- a/Engine/Plugins/Experimental/PythonScriptPlugin/Source/PythonScriptPlugin/Private/PythonScriptPlugin.cpp
+++ b/Engine/Plugins/Experimental/PythonScriptPlugin/Source/PythonScriptPlugin/Private/PythonScriptPlugin.cpp
@@ -654,7 +654,7 @@ void FPythonScriptPlugin::InitializePython()
 		PySys_SetArgvEx(1, NullPyArgPtrs, 0);
 
 		// Enable developer warnings if requested
-		if (GetDefault<UPythonScriptPluginSettings>()->bDeveloperMode)
+		if (IsDeveloperModeEnabled())
 		{
 			PyUtil::EnableDeveloperWarnings();
 		}
@@ -836,7 +836,7 @@ void FPythonScriptPlugin::RequestStubCodeGeneration()
 
 void FPythonScriptPlugin::GenerateStubCode()
 {
-	if (GetDefault<UPythonScriptPluginSettings>()->bDeveloperMode)
+	if (IsDeveloperModeEnabled())
 	{
 		// Generate stub code if developer mode enabled
 		FPyWrapperTypeRegistry::Get().GenerateStubCodeForWrappedTypes();
@@ -1176,6 +1176,11 @@ void FPythonScriptPlugin::OnContentPathDismounted(const FString& InAssetPath, co
 	PyUtil::RemoveSystemPath(FPaths::ConvertRelativePathToFull(InFilesystemPath / TEXT("Python")));
 }
 
+bool FPythonScriptPlugin::IsDeveloperModeEnabled()
+{
+	return GetDefault<UPythonScriptPluginSettings>()->bDeveloperMode || GetDefault<UPythonScriptPluginUserSettings>()->bDeveloperMode;
+}
+
 void FPythonScriptPlugin::OnAssetRenamed(const FAssetData& Data, const FString& OldName)
 {
 	const FName OldPackageName = *FPackageName::ObjectPathToPackageName(OldName);
diff --git a/Engine/Plugins/Experimental/PythonScriptPlugin/Source/PythonScriptPlugin/Private/PythonScriptPlugin.h b/Engine/Plugins/Experimental/PythonScriptPlugin/Source/PythonScriptPlugin/Private/PythonScriptPlugin.h
--- a/Engine/Plugins/Experimental/PythonScriptPlugin/Source/PythonScriptPlugin/Private/PythonScriptPlugin.h
+++ b/Engine/Plugins/Experimental/PythonScriptPlugin/Source/PythonScriptPlugin/Private/PythonScriptPlugin.h
@@ -170,6 +170,8 @@ private:
 
 	void OnContentPathDismounted(const FString& InAssetPath, const FString& InFilesystemPath);
 
+	static bool IsDeveloperModeEnabled();
+
 	void OnAssetRenamed(const FAssetData& Data, const FString& OldName);
 
 	void OnAssetRemoved(const FAssetData& Data);
diff --git a/Engine/Plugins/Experimental/PythonScriptPlugin/Source/PythonScriptPlugin/Private/PythonScriptPluginSettings.cpp b/Engine/Plugins/Experimental/PythonScriptPlugin/Source/PythonScriptPlugin/Private/PythonScriptPluginSettings.cpp
--- a/Engine/Plugins/Experimental/PythonScriptPlugin/Source/PythonScriptPlugin/Private/PythonScriptPluginSettings.cpp
+++ b/Engine/Plugins/Experimental/PythonScriptPlugin/Source/PythonScriptPlugin/Private/PythonScriptPluginSettings.cpp
@@ -67,5 +67,19 @@ FText UPythonScriptPluginSettings::GetSectionText() const
 
 #endif	// WITH_EDITOR
 
+UPythonScriptPluginUserSettings::UPythonScriptPluginUserSettings()
+{
+	CategoryName = TEXT("Plugins");
+	SectionName = TEXT("Python");
+}
+
+#if WITH_EDITOR
+FText UPythonScriptPluginUserSettings::GetSectionText() const
+{
+	return LOCTEXT("UserSettingsDisplayName", "Python");
+}
+
+#endif	// WITH_EDITOR
+
 #undef LOCTEXT_NAMESPACE
 
diff --git a/Engine/Plugins/Experimental/PythonScriptPlugin/Source/PythonScriptPlugin/Private/PythonScriptPluginSettings.h b/Engine/Plugins/Experimental/PythonScriptPlugin/Source/PythonScriptPlugin/Private/PythonScriptPluginSettings.h
--- a/Engine/Plugins/Experimental/PythonScriptPlugin/Source/PythonScriptPlugin/Private/PythonScriptPluginSettings.h
+++ b/Engine/Plugins/Experimental/PythonScriptPlugin/Source/PythonScriptPlugin/Private/PythonScriptPluginSettings.h
@@ -35,8 +35,14 @@ public:
 	UPROPERTY(config, EditAnywhere, Category=Python, meta=(ConfigRestartRequired=true, RelativePath))
 	TArray<FDirectoryPath> AdditionalPaths;
 
-	/** Should Developer Mode be enabled on the Python interpreter (will enable extra warnings (eg, for deprecated code), and enable stub code generation for use with external IDEs). */
-	UPROPERTY(config, EditAnywhere, Category=Python, meta=(ConfigRestartRequired=true))
+	/**
+	 * Should Developer Mode be enabled on the Python interpreter *for all users of the project*
+	 * Note: Most of the time you want to enable bDeveloperMode in the Editor Preferences instead!
+	 *
+	 * (will also enable extra warnings (e.g., for deprecated code), and enable stub code generation for
+	 * use with external IDEs).
+	 */
+	UPROPERTY(config, EditAnywhere, Category=Python, meta=(ConfigRestartRequired=true, DisplayName="Developer Mode (all users)"), AdvancedDisplay)
 	bool bDeveloperMode;
 
 	/** Should remote Python execution be enabled? */
@@ -63,3 +69,27 @@ public:
 	UPROPERTY(config, EditAnywhere, Category=PythonRemoteExecution, AdvancedDisplay, meta=(DisplayName="Multicast Time-To-Live"))
 	uint8 RemoteExecutionMulticastTtl;
 };
+
+
+UCLASS(config=EditorPerProjectUserSettings)
+class UPythonScriptPluginUserSettings : public UDeveloperSettings
+{
+	GENERATED_BODY()
+
+public:
+	UPythonScriptPluginUserSettings();
+
+#if WITH_EDITOR
+	//~ UDeveloperSettings interface
+	virtual FText GetSectionText() const override;
+#endif
+
+	/**
+	 * Should Developer Mode be enabled on the Python interpreter?
+	 *
+	 * (will also enable extra warnings (e.g., for deprecated code), and enable stub code generation for
+	 * use with external IDEs).
+	 */
+	UPROPERTY(config, EditAnywhere, Category=Python, meta=(ConfigRestartRequired=true))
+	bool bDeveloperMode;
+};
\ No newline at end of file
