Avoid a crash from decrementing a Python object ref-count outside of the GIL

We now keep the GIL locked for the duration that the PySelf pointer is being used

#rnx

diff --git a/Engine/Plugins/Experimental/PythonScriptPlugin/Source/PythonScriptPlugin/Private/PyWrapperObject.cpp b/Engine/Plugins/Experimental/PythonScriptPlugin/Source/PythonScriptPlugin/Private/PyWrapperObject.cpp
--- a/Engine/Plugins/Experimental/PythonScriptPlugin/Source/PythonScriptPlugin/Private/PyWrapperObject.cpp
+++ b/Engine/Plugins/Experimental/PythonScriptPlugin/Source/PythonScriptPlugin/Private/PyWrapperObject.cpp
@@ -1965,23 +1965,23 @@ DEFINE_FUNCTION(UPythonGeneratedClass::CallPythonFunction)
 		}
 	}
 
-	// Find the Python object to call the function on
-	FPyObjectPtr PySelf;
-	bool bSelfError = false;
-	if (!Func->HasAnyFunctionFlags(FUNC_Static))
+	// Execute Python code within this block
 	{
 		FPyScopedGIL GIL;
-		PySelf = FPyObjectPtr::StealReference((PyObject*)FPyWrapperObjectFactory::Get().CreateInstance(P_THIS_OBJECT));
-		if (!PySelf)
+
+		// Find the Python object to call the function on
+		FPyObjectPtr PySelf;
+		bool bSelfError = false;
+		if (!Func->HasAnyFunctionFlags(FUNC_Static))
 		{
-			UE_LOG(LogPython, Error, TEXT("Failed to create a Python wrapper for '%s'"), *P_THIS_OBJECT->GetName());
-			bSelfError = true;
+			PySelf = FPyObjectPtr::StealReference((PyObject*)FPyWrapperObjectFactory::Get().CreateInstance(P_THIS_OBJECT));
+			if (!PySelf)
+			{
+				UE_LOG(LogPython, Error, TEXT("Failed to create a Python wrapper for '%s'"), *P_THIS_OBJECT->GetName());
+				bSelfError = true;
+			}
 		}
-	}
-
-	// Execute Python code within this block
-	{
-		FPyScopedGIL GIL;
+	
 		if (!PyGenUtil::InvokePythonCallableFromUnrealFunctionThunk(PySelf, FuncDef ? FuncDef->PyFunction.GetPtr() : nullptr, Func, Context, Stack, RESULT_PARAM) || bSelfError)
 		{
 			PyUtil::ReThrowPythonError();
