Fixed ExecutePythonScript reporting an error if the node had no outputs

The internal MakeArray node would default to adding 1 pin, even if the input array of pin names was empty

diff --git a/Engine/Plugins/Experimental/PythonScriptPlugin/Source/PythonScriptPlugin/Private/K2Node_ExecutePythonScript.cpp b/Engine/Plugins/Experimental/PythonScriptPlugin/Source/PythonScriptPlugin/Private/K2Node_ExecutePythonScript.cpp
--- a/Engine/Plugins/Experimental/PythonScriptPlugin/Source/PythonScriptPlugin/Private/K2Node_ExecutePythonScript.cpp
+++ b/Engine/Plugins/Experimental/PythonScriptPlugin/Source/PythonScriptPlugin/Private/K2Node_ExecutePythonScript.cpp
@@ -244,6 +244,7 @@ void UK2Node_ExecutePythonScript::ExpandNode(FKismetCompilerContext& CompilerCon
 	{
 		// Create a "Make Array" node to compile the list of arguments into an array for the Format function being called
 		UK2Node_MakeArray* MakeArrayNode = CompilerContext.SpawnIntermediateNode<UK2Node_MakeArray>(this, SourceGraph);
+		MakeArrayNode->NumInputs = ArgNames.Num();
 		MakeArrayNode->AllocateDefaultPins();
 
 		// Connect the output of the "Make Array" pin to the destination input pin
@@ -267,12 +268,6 @@ void UK2Node_ExecutePythonScript::ExpandNode(FKismetCompilerContext& CompilerCon
 			UEdGraphPin* MakeLiteralIntValuePin = MakeLiteralStringNode->FindPinChecked(TEXT("Value"), EGPD_Input);
 			MakeLiteralIntValuePin->DefaultValue = ExecutePythonScriptUtil::PythonizePinName(ArgName);
 
-			// The "Make Array" node already has one pin available, so don't create one for ArgIdx == 0
-			if (ArgIndex > 0)
-			{
-				MakeArrayNode->AddInputPin();
-			}
-
 			// Find the input pin on the "Make Array" node by index and link it to the literal string
 			UEdGraphPin* ArrayIn = MakeArrayNode->FindPinChecked(FString::Printf(TEXT("[%d]"), ArgIndex++));
 			MakeLiteralStringNode->GetReturnValuePin()->MakeLinkTo(ArrayIn);
