Python Script Plugin File Menu Icons

diff --git a/Engine/Plugins/Experimental/PythonScriptPlugin/Source/PythonScriptPlugin/Private/PythonScriptPlugin.cpp b/Engine/Plugins/Experimental/PythonScriptPlugin/Source/PythonScriptPlugin/Private/PythonScriptPlugin.cpp
--- a/Engine/Plugins/Experimental/PythonScriptPlugin/Source/PythonScriptPlugin/Private/PythonScriptPlugin.cpp
+++ b/Engine/Plugins/Experimental/PythonScriptPlugin/Source/PythonScriptPlugin/Private/PythonScriptPlugin.cpp
@@ -2,6 +2,7 @@
 
 #include "PythonScriptPlugin.h"
 #include "PythonScriptPluginSettings.h"
+#include "PythonScriptPluginStyle.h"
 #include "PythonScriptRemoteExecution.h"
 #include "PyGIL.h"
 #include "PyCore.h"
@@ -350,16 +351,15 @@ private:
 			"OpenPython",
 			LOCTEXT("OpenPython", "Execute Python Script..."),
 			LOCTEXT("OpenPythonTooltip", "Open a Python Script file and Execute it."),
-			FSlateIcon(),
+			FSlateIcon(FPythonScriptPluginEditorStyle::Get().GetStyleSetName(), "Icons.PythonExecute"),
 			FUIAction(FExecuteAction::CreateRaw(this, &FPythonCommandMenuImpl::Menu_ExecutePython))
 		);
 		Section.AddSubMenu(
 			"RecentPythonsSubMenu",
 			LOCTEXT("RecentPythonsSubMenu", "Recent Python Scripts"),
 			LOCTEXT("RecentPythonsSubMenu_ToolTip", "Select a recent Python Script file and Execute it."),
-			FNewToolMenuDelegate::CreateRaw(this, &FPythonCommandMenuImpl::MakeRecentPythonScriptMenu),
-			false,
-			FSlateIcon(FEditorStyle::GetStyleSetName(), "MainFrame.RecentLevels")
+			FNewToolMenuDelegate::CreateRaw(this, &FPythonCommandMenuImpl::MakeRecentPythonScriptMenu), false,
+			FSlateIcon(FPythonScriptPluginEditorStyle::Get().GetStyleSetName(), "Icons.PythonRecent")
 		);
 	}
 
@@ -501,6 +501,7 @@ void FPythonScriptPlugin::StartupModule()
 	RemoteExecution = MakeUnique<FPythonScriptRemoteExecution>(this);
 
 #if WITH_EDITOR
+	FPythonScriptPluginEditorStyle::Get();
 	FCoreDelegates::OnPostEngineInit.AddRaw(this, &FPythonScriptPlugin::OnPostEngineInit);
 #endif // WITH_EDITOR
 
diff --git a/Engine/Plugins/Experimental/PythonScriptPlugin/Source/PythonScriptPlugin/Private/PythonScriptPluginStyle.cpp b/Engine/Plugins/Experimental/PythonScriptPlugin/Source/PythonScriptPlugin/Private/PythonScriptPluginStyle.cpp
new file mode 100644
--- /dev/null
+++ b/Engine/Plugins/Experimental/PythonScriptPlugin/Source/PythonScriptPlugin/Private/PythonScriptPluginStyle.cpp
@@ -0,0 +1,37 @@
+// Copyright Epic Games, Inc. All Rights Reserved.
+
+#include "PythonScriptPluginStyle.h"
+
+#include "Styling/SlateTypes.h"
+#include "Styling/CoreStyle.h"
+#include "EditorStyleSet.h"
+#include "Styling/SlateStyleMacros.h"
+
+FName FPythonScriptPluginEditorStyle::StyleName("PythonScriptPluginEditorStyle");
+
+FPythonScriptPluginEditorStyle::FPythonScriptPluginEditorStyle()
+	: FSlateStyleSet(StyleName)
+{
+	const FVector2D IconSize(16.0f, 16.0f);
+
+	SetContentRoot(FPaths::EnginePluginsDir() / TEXT("Experimental/PythonScriptPlugin/Content"));
+	SetCoreContentRoot(FPaths::EngineContentDir() / TEXT("Slate"));
+
+	Set("Icons.PythonExecute", new IMAGE_BRUSH_SVG("PythonExecute", IconSize));
+	Set("Icons.PythonRecent",  new IMAGE_BRUSH_SVG("PythonRecent", IconSize));
+
+	FSlateStyleRegistry::RegisterSlateStyle(*this);
+}
+
+FPythonScriptPluginEditorStyle::~FPythonScriptPluginEditorStyle()
+{
+	FSlateStyleRegistry::UnRegisterSlateStyle(*this);
+}
+
+FPythonScriptPluginEditorStyle& FPythonScriptPluginEditorStyle::Get()
+{
+	static FPythonScriptPluginEditorStyle Inst;
+	return Inst;
+}
+
+
diff --git a/Engine/Plugins/Experimental/PythonScriptPlugin/Source/PythonScriptPlugin/Private/PythonScriptPluginStyle.h b/Engine/Plugins/Experimental/PythonScriptPlugin/Source/PythonScriptPlugin/Private/PythonScriptPluginStyle.h
new file mode 100644
--- /dev/null
+++ b/Engine/Plugins/Experimental/PythonScriptPlugin/Source/PythonScriptPlugin/Private/PythonScriptPluginStyle.h
@@ -0,0 +1,24 @@
+// Copyright Epic Games, Inc. All Rights Reserved.
+
+#pragma once
+
+#include "Styling/SlateStyle.h"
+#include "Styling/SlateStyleRegistry.h"
+
+/**
+ * Slate style set that defines the styles for Python UI in the Editor 
+ */
+class FPythonScriptPluginEditorStyle
+	: public FSlateStyleSet
+{
+public:
+	static FName StyleName;
+
+	/** Access the singleton instance for this style set */
+	static FPythonScriptPluginEditorStyle& Get();
+
+private:
+
+	FPythonScriptPluginEditorStyle();
+	~FPythonScriptPluginEditorStyle();
+};
diff --git a/Engine/Plugins/Experimental/PythonScriptPlugin/Source/PythonScriptPlugin/PythonScriptPlugin.Build.cs b/Engine/Plugins/Experimental/PythonScriptPlugin/Source/PythonScriptPlugin/PythonScriptPlugin.Build.cs
--- a/Engine/Plugins/Experimental/PythonScriptPlugin/Source/PythonScriptPlugin/PythonScriptPlugin.Build.cs
+++ b/Engine/Plugins/Experimental/PythonScriptPlugin/Source/PythonScriptPlugin/PythonScriptPlugin.Build.cs
@@ -19,6 +19,7 @@ namespace UnrealBuildTool.Rules
 				new string[] {
 					"Analytics",
 					"AssetRegistry",
+					"EditorStyle",
 					"Projects",
 					"Python",
 					"Slate",
