Merging PythonScriptPlugin from //UE5/Release-Engine-Staging to Main (//UE5/Main)

#rnx

diff --git a/Engine/Plugins/Experimental/PythonScriptPlugin/Source/PythonScriptPlugin/Private/PyCore.cpp b/Engine/Plugins/Experimental/PythonScriptPlugin/Source/PythonScriptPlugin/Private/PyCore.cpp
--- a/Engine/Plugins/Experimental/PythonScriptPlugin/Source/PythonScriptPlugin/Private/PyCore.cpp
+++ b/Engine/Plugins/Experimental/PythonScriptPlugin/Source/PythonScriptPlugin/Private/PyCore.cpp
@@ -1584,7 +1584,9 @@ PyObject* GenerateEnum(PyObject* InSelf, PyObject* InArgs)
 
 PyObject* FlushGeneratedTypeReinstancing(PyObject* InSelf)
 {
+	Py_BEGIN_ALLOW_THREADS
 	FPyWrapperTypeReinstancer::Get().ProcessPending();
+	Py_END_ALLOW_THREADS
 
 	Py_RETURN_NONE;
 }
diff --git a/Engine/Plugins/Experimental/PythonScriptPlugin/Source/PythonScriptPlugin/Private/PyGenUtil.cpp b/Engine/Plugins/Experimental/PythonScriptPlugin/Source/PythonScriptPlugin/Private/PyGenUtil.cpp
--- a/Engine/Plugins/Experimental/PythonScriptPlugin/Source/PythonScriptPlugin/Private/PyGenUtil.cpp
+++ b/Engine/Plugins/Experimental/PythonScriptPlugin/Source/PythonScriptPlugin/Private/PyGenUtil.cpp
@@ -200,7 +200,7 @@ const FGeneratedWrappedOperatorSignature& FGeneratedWrappedOperatorSignature::Op
 }
 
 bool FGeneratedWrappedOperatorSignature::StringToSignature(const TCHAR* InStr, FGeneratedWrappedOperatorSignature& OutSignature)
-	{
+{
 	for (int32 OpTypeIndex = 0; OpTypeIndex < (int32)EGeneratedWrappedOperatorType::Num; ++OpTypeIndex)
 	{
 		const FGeneratedWrappedOperatorSignature& PotentialSignature = OpTypeToSignature((EGeneratedWrappedOperatorType)OpTypeIndex);
@@ -710,7 +710,7 @@ void FGeneratedWrappedClassType::Finalize_PostReady()
 	// Execute Python code within this block
 	{
 		FPyScopedGIL GIL;
-	FPyMethodWithClosureDef::AddMethods(Methods.PyMethods.GetData(), &PyType);
+		FPyMethodWithClosureDef::AddMethods(Methods.PyMethods.GetData(), &PyType);
 		FPyConstantDef::AddConstantsToType(Constants.PyConstants.GetData(), &PyType);
 	}
 }
@@ -744,10 +744,14 @@ void FGeneratedWrappedEnumType::Finalize_PostReady()
 
 void FGeneratedWrappedEnumType::Reset_CleansePyType()
 {
-	// Unregister the existing enum entries
-	for (const FGeneratedWrappedEnumEntry& EnumEntry : EnumEntries)
+	// Execute Python code within this block
 	{
-		PyDict_DelItemString(PyType.tp_dict, EnumEntry.EntryName.GetData());
+		// Unregister the existing enum entries
+		FPyScopedGIL GIL;
+		for (const FGeneratedWrappedEnumEntry& EnumEntry : EnumEntries)
+		{
+			PyDict_DelItemString(PyType.tp_dict, EnumEntry.EntryName.GetData());
+		}
 	}
 
 	FGeneratedWrappedType::Reset_CleansePyType();
@@ -2329,8 +2333,11 @@ void PythonizeStructValueImpl(const UScriptStruct* InStruct, const void* InStruc
 					void* SelfArgInstance = SelfParam.ParamProp->ContainerPtrToValuePtr<void>(FuncParams.GetStructMemory());
 					CastFieldChecked<const FStructProperty>(SelfParam.ParamProp)->Struct->CopyScriptStruct(SelfArgInstance, InStructValue);
 				}
-				PyUtil::InvokeFunctionCall(Obj, BreakFuncDef.Func, FuncParams.GetStructMemory(), TEXT("pythonize default struct value"));
-				PyErr_Clear(); // Clear any errors in case InvokeFunctionCall failed
+				{
+					FPyScopedGIL GIL;
+					PyUtil::InvokeFunctionCall(Obj, BreakFuncDef.Func, FuncParams.GetStructMemory(), TEXT("pythonize default struct value"));
+					PyErr_Clear(); // Clear any errors in case InvokeFunctionCall failed
+				}
 
 				// Extract the output argument values as defaults for the struct
 				for (int32 OuputParamIndex = 0; OuputParamIndex < BreakFuncDef.OutputParams.Num(); ++OuputParamIndex)
diff --git a/Engine/Plugins/Experimental/PythonScriptPlugin/Source/PythonScriptPlugin/Private/PyReferenceCollector.cpp b/Engine/Plugins/Experimental/PythonScriptPlugin/Source/PythonScriptPlugin/Private/PyReferenceCollector.cpp
--- a/Engine/Plugins/Experimental/PythonScriptPlugin/Source/PythonScriptPlugin/Private/PyReferenceCollector.cpp
+++ b/Engine/Plugins/Experimental/PythonScriptPlugin/Source/PythonScriptPlugin/Private/PyReferenceCollector.cpp
@@ -7,6 +7,7 @@
 #include "PyWrapperStruct.h"
 #include "PyWrapperEnum.h"
 #include "PyWrapperDelegate.h"
+#include "PyGIL.h"
 #include "UObject/UnrealType.h"
 #include "UObject/UObjectHash.h"
 #include "UObject/PurgingReferenceCollector.h"
@@ -31,12 +32,15 @@ void FPyReferenceCollector::RemoveWrappedInstance(FPyWrapperBase* InInstance)
 
 void FPyReferenceCollector::AddReferencedObjects(FReferenceCollector& InCollector)
 {
-	for (FPyWrapperBase* PythonWrappedInstance : PythonWrappedInstances)
 	{
-		FPyWrapperBaseMetaData* PythonWrappedInstanceMetaData = FPyWrapperBaseMetaData::GetMetaData(PythonWrappedInstance);
-		if (PythonWrappedInstanceMetaData)
+		FPyScopedGIL GIL;
+		for (FPyWrapperBase* PythonWrappedInstance : PythonWrappedInstances)
 		{
-			PythonWrappedInstanceMetaData->AddReferencedObjects(PythonWrappedInstance, InCollector);
+			FPyWrapperBaseMetaData* PythonWrappedInstanceMetaData = FPyWrapperBaseMetaData::GetMetaData(PythonWrappedInstance);
+			if (PythonWrappedInstanceMetaData)
+			{
+				PythonWrappedInstanceMetaData->AddReferencedObjects(PythonWrappedInstance, InCollector);
+			}
 		}
 	}
 
@@ -157,8 +161,10 @@ void FPyReferenceCollector::PurgeUnrealGeneratedTypes()
 
 	if (PurgingReferenceCollector.HasObjectToPurge())
 	{
+		Py_BEGIN_ALLOW_THREADS
 		AddReferencedObjects(PurgingReferenceCollector);
 		CollectGarbage(GARBAGE_COLLECTION_KEEPFLAGS);
+		Py_END_ALLOW_THREADS
 
 		for (const FWeakObjectPtr& WeakReferencesToPurgedObject : WeakReferencesToPurgedObjects)
 		{
diff --git a/Engine/Plugins/Experimental/PythonScriptPlugin/Source/PythonScriptPlugin/Private/PyUtil.cpp b/Engine/Plugins/Experimental/PythonScriptPlugin/Source/PythonScriptPlugin/Private/PyUtil.cpp
--- a/Engine/Plugins/Experimental/PythonScriptPlugin/Source/PythonScriptPlugin/Private/PyUtil.cpp
+++ b/Engine/Plugins/Experimental/PythonScriptPlugin/Source/PythonScriptPlugin/Private/PyUtil.cpp
@@ -3,6 +3,7 @@
 #include "PyUtil.h"
 #include "PyCore.h"
 #include "PyConversion.h"
+#include "PyGIL.h"
 
 #include "PyWrapperObject.h"
 #include "PyWrapperStruct.h"
@@ -554,11 +555,13 @@ bool InvokeFunctionCall(UObject* InObj, const UFunction* InFunc, void* InBasePar
 	{
 		if (Verbosity == ELogVerbosity::Error)
 		{
+			FPyScopedGIL GIL;
 			SetPythonError(PyExc_RuntimeError, InErrorCtxt, ExceptionMessage);
 			bThrewException = true;
 		}
 		else if (Verbosity == ELogVerbosity::Warning)
 		{
+			FPyScopedGIL GIL;
 			if (SetPythonWarning(PyExc_RuntimeWarning, InErrorCtxt, ExceptionMessage) == -1)
 			{
 				// -1 from SetPythonWarning means the warning should be an exception
@@ -574,7 +577,9 @@ bool InvokeFunctionCall(UObject* InObj, const UFunction* InFunc, void* InBasePar
 	});
 
 	FEditorScriptExecutionGuard ScriptGuard;
+	Py_BEGIN_ALLOW_THREADS
 	InObj->ProcessEvent((UFunction*)InFunc, InBaseParamsAddr);
+	Py_END_ALLOW_THREADS
 
 	return !bThrewException;
 }
diff --git a/Engine/Plugins/Experimental/PythonScriptPlugin/Source/PythonScriptPlugin/Private/PyWrapperBase.cpp b/Engine/Plugins/Experimental/PythonScriptPlugin/Source/PythonScriptPlugin/Private/PyWrapperBase.cpp
--- a/Engine/Plugins/Experimental/PythonScriptPlugin/Source/PythonScriptPlugin/Private/PyWrapperBase.cpp
+++ b/Engine/Plugins/Experimental/PythonScriptPlugin/Source/PythonScriptPlugin/Private/PyWrapperBase.cpp
@@ -2,6 +2,7 @@
 
 #include "PyWrapperBase.h"
 #include "PyReferenceCollector.h"
+#include "PyGIL.h"
 
 #if WITH_PYTHON
 
@@ -114,6 +115,8 @@ void FPyWrapperBaseMetaData::SetMetaData(PyTypeObject* PyType, FPyWrapperBaseMet
 {
 	if (PyType && PyType->tp_dict)
 	{
+		FPyScopedGIL GIL;
+
 		FPyWrapperBaseMetaDataObject* PyWrapperMetaData = (FPyWrapperBaseMetaDataObject*)PyDict_GetItemString(PyType->tp_dict, "_wrapper_meta_data");
 		if (!PyWrapperMetaData)
 		{
@@ -129,6 +132,8 @@ FPyWrapperBaseMetaData* FPyWrapperBaseMetaData::GetMetaData(PyTypeObject* PyType
 {
 	if (PyType && PyType->tp_dict)
 	{
+		FPyScopedGIL GIL;
+
 		FPyWrapperBaseMetaDataObject* PyWrapperMetaData = (FPyWrapperBaseMetaDataObject*)PyDict_GetItemString(PyType->tp_dict, "_wrapper_meta_data");
 		if (PyWrapperMetaData)
 		{
diff --git a/Engine/Plugins/Experimental/PythonScriptPlugin/Source/PythonScriptPlugin/Private/PyWrapperDelegate.cpp b/Engine/Plugins/Experimental/PythonScriptPlugin/Source/PythonScriptPlugin/Private/PyWrapperDelegate.cpp
--- a/Engine/Plugins/Experimental/PythonScriptPlugin/Source/PythonScriptPlugin/Private/PyWrapperDelegate.cpp
+++ b/Engine/Plugins/Experimental/PythonScriptPlugin/Source/PythonScriptPlugin/Private/PyWrapperDelegate.cpp
@@ -35,6 +35,12 @@ DEFINE_FUNCTION(UPythonCallableForDelegate::CallPythonNative)
 
 #if WITH_PYTHON
 void UPythonCallableForDelegate::BeginDestroy()
+{
+	ReleasePythonResources();
+	Super::BeginDestroy();
+}
+
+void UPythonCallableForDelegate::ReleasePythonResources()
 {
 	// This may be called after Python has already shut down
 	if (Py_IsInitialized())
@@ -44,16 +50,9 @@ void UPythonCallableForDelegate::BeginDestroy()
 	}
 	else
 	{
-		// Release ownership if Python has been shut down to avoid attempting to delete the callable (which is already dead)
+		// Release ownership if Python has been shut down to avoid attempting to delete the object (which is already dead)
 		PyCallable.Release();
 	}
-
-	Super::BeginDestroy();
-}
-
-void UPythonCallableForDelegate::ReleasePythonResources()
-{
-	PyCallable.Reset();
 }
 
 PyObject* UPythonCallableForDelegate::GetCallable() const
diff --git a/Engine/Plugins/Experimental/PythonScriptPlugin/Source/PythonScriptPlugin/Private/PyWrapperEnum.cpp b/Engine/Plugins/Experimental/PythonScriptPlugin/Source/PythonScriptPlugin/Private/PyWrapperEnum.cpp
--- a/Engine/Plugins/Experimental/PythonScriptPlugin/Source/PythonScriptPlugin/Private/PyWrapperEnum.cpp
+++ b/Engine/Plugins/Experimental/PythonScriptPlugin/Source/PythonScriptPlugin/Private/PyWrapperEnum.cpp
@@ -2,6 +2,7 @@
 
 #include "PyWrapperEnum.h"
 #include "PyWrapperTypeRegistry.h"
+#include "PyGIL.h"
 #include "PyCore.h"
 #include "PyUtil.h"
 #include "PyConversion.h"
@@ -752,7 +753,9 @@ public:
 			NewEnum->ClearFlags(RF_Public | RF_Standalone);
 			NewEnum = nullptr;
 
+			Py_BEGIN_ALLOW_THREADS
 			CollectGarbage(GARBAGE_COLLECTION_KEEPFLAGS);
+			Py_END_ALLOW_THREADS
 		}
 	}
 
@@ -859,9 +862,26 @@ private:
 	bool bDidExist;
 };
 
+void UPythonGeneratedEnum::BeginDestroy()
+{
+	ReleasePythonResources();
+	Super::BeginDestroy();
+}
+
 void UPythonGeneratedEnum::ReleasePythonResources()
 {
-	PyType.Reset();
+	// This may be called after Python has already shut down
+	if (Py_IsInitialized())
+	{
+		FPyScopedGIL GIL;
+		PyType.Reset();
+	}
+	else
+	{
+		// Release ownership if Python has been shut down to avoid attempting to delete the object (which is already dead)
+		PyType.Release();
+	}
+
 	EnumValueDefs.Reset();
 	PyMetaData = FPyWrapperEnumMetaData();
 }
diff --git a/Engine/Plugins/Experimental/PythonScriptPlugin/Source/PythonScriptPlugin/Private/PyWrapperEnum.h b/Engine/Plugins/Experimental/PythonScriptPlugin/Source/PythonScriptPlugin/Private/PyWrapperEnum.h
--- a/Engine/Plugins/Experimental/PythonScriptPlugin/Source/PythonScriptPlugin/Private/PyWrapperEnum.h
+++ b/Engine/Plugins/Experimental/PythonScriptPlugin/Source/PythonScriptPlugin/Private/PyWrapperEnum.h
@@ -120,6 +120,9 @@ class UPythonGeneratedEnum : public UEnum, public IPythonResourceOwner
 #if WITH_PYTHON
 
 public:
+	//~ UObject interface
+	virtual void BeginDestroy() override;
+
 	//~ IPythonResourceOwner interface
 	virtual void ReleasePythonResources() override;
 
diff --git a/Engine/Plugins/Experimental/PythonScriptPlugin/Source/PythonScriptPlugin/Private/PyWrapperObject.cpp b/Engine/Plugins/Experimental/PythonScriptPlugin/Source/PythonScriptPlugin/Private/PyWrapperObject.cpp
--- a/Engine/Plugins/Experimental/PythonScriptPlugin/Source/PythonScriptPlugin/Private/PyWrapperObject.cpp
+++ b/Engine/Plugins/Experimental/PythonScriptPlugin/Source/PythonScriptPlugin/Private/PyWrapperObject.cpp
@@ -1241,7 +1241,9 @@ public:
 			NewClass->ClearFlags(RF_Public | RF_Standalone);
 			NewClass = nullptr;
 
+			Py_BEGIN_ALLOW_THREADS
 			CollectGarbage(GARBAGE_COLLECTION_KEEPFLAGS);
+			Py_END_ALLOW_THREADS
 		}
 	}
 
@@ -1839,11 +1841,36 @@ void UPythonGeneratedClass::PostInitInstance(UObject* InObj)
 	}
 }
 
+void UPythonGeneratedClass::BeginDestroy()
+{
+	ReleasePythonResources();
+	Super::BeginDestroy();
+}
 
 void UPythonGeneratedClass::ReleasePythonResources()
 {
-	PyType.Reset();
-	PyPostInitFunction.Reset();
+	// This may be called after Python has already shut down
+	if (Py_IsInitialized())
+	{
+		FPyScopedGIL GIL;
+		PyType.Reset();
+		PyPostInitFunction.Reset();
+		for (const TSharedPtr<PyGenUtil::FFunctionDef>& FunctionDef : FunctionDefs)
+		{
+			FunctionDef->PyFunction.Reset();
+		}
+	}
+	else
+	{
+		// Release ownership if Python has been shut down to avoid attempting to delete the objects (which are already dead)
+		PyType.Release();
+		PyPostInitFunction.Release();
+		for (const TSharedPtr<PyGenUtil::FFunctionDef>& FunctionDef : FunctionDefs)
+		{
+			FunctionDef->PyFunction.Release();
+		}
+	}
+
 	PropertyDefs.Reset();
 	FunctionDefs.Reset();
 	PyMetaData = FPyWrapperObjectMetaData();
@@ -2014,6 +2041,7 @@ DEFINE_FUNCTION(UPythonGeneratedClass::CallPythonFunction)
 		{
 			PyUtil::ReThrowPythonError();
 		}
+		PySelf.Reset(); // Need to reset this while still under the GIL
 	}
 }
 
diff --git a/Engine/Plugins/Experimental/PythonScriptPlugin/Source/PythonScriptPlugin/Private/PyWrapperObject.h b/Engine/Plugins/Experimental/PythonScriptPlugin/Source/PythonScriptPlugin/Private/PyWrapperObject.h
--- a/Engine/Plugins/Experimental/PythonScriptPlugin/Source/PythonScriptPlugin/Private/PyWrapperObject.h
+++ b/Engine/Plugins/Experimental/PythonScriptPlugin/Source/PythonScriptPlugin/Private/PyWrapperObject.h
@@ -190,6 +190,7 @@ class UPythonGeneratedClass : public UClass, public IPythonResourceOwner
 public:
 	//~ UObject interface
 	virtual void PostRename(UObject* OldOuter, const FName OldName) override;
+	virtual void BeginDestroy() override;
 
 	//~ UClass interface
 	virtual void PostInitInstance(UObject* InObj) override;
diff --git a/Engine/Plugins/Experimental/PythonScriptPlugin/Source/PythonScriptPlugin/Private/PyWrapperStruct.cpp b/Engine/Plugins/Experimental/PythonScriptPlugin/Source/PythonScriptPlugin/Private/PyWrapperStruct.cpp
--- a/Engine/Plugins/Experimental/PythonScriptPlugin/Source/PythonScriptPlugin/Private/PyWrapperStruct.cpp
+++ b/Engine/Plugins/Experimental/PythonScriptPlugin/Source/PythonScriptPlugin/Private/PyWrapperStruct.cpp
@@ -1401,7 +1401,9 @@ public:
 			NewStruct->ClearFlags(RF_Public | RF_Standalone);
 			NewStruct = nullptr;
 
+			Py_BEGIN_ALLOW_THREADS
 			CollectGarbage(GARBAGE_COLLECTION_KEEPFLAGS);
+			Py_END_ALLOW_THREADS
 		}
 	}
 
@@ -1593,10 +1595,28 @@ void UPythonGeneratedStruct::InitializeStruct(void* Dest, int32 ArrayDim) const
 	}
 }
 
+void UPythonGeneratedStruct::BeginDestroy()
+{
+	ReleasePythonResources();
+	Super::BeginDestroy();
+}
+
 void UPythonGeneratedStruct::ReleasePythonResources()
 {
-	PyType.Reset();
-	PyPostInitFunction.Reset();
+	// This may be called after Python has already shut down
+	if (Py_IsInitialized())
+	{
+		FPyScopedGIL GIL;
+		PyType.Reset();
+		PyPostInitFunction.Reset();
+	}
+	else
+	{
+		// Release ownership if Python has been shut down to avoid attempting to delete the objects (which are already dead)
+		PyType.Release();
+		PyPostInitFunction.Release();
+	}
+
 	PropertyDefs.Reset();
 	PyMetaData = FPyWrapperStructMetaData();
 }
diff --git a/Engine/Plugins/Experimental/PythonScriptPlugin/Source/PythonScriptPlugin/Private/PyWrapperStruct.h b/Engine/Plugins/Experimental/PythonScriptPlugin/Source/PythonScriptPlugin/Private/PyWrapperStruct.h
--- a/Engine/Plugins/Experimental/PythonScriptPlugin/Source/PythonScriptPlugin/Private/PyWrapperStruct.h
+++ b/Engine/Plugins/Experimental/PythonScriptPlugin/Source/PythonScriptPlugin/Private/PyWrapperStruct.h
@@ -378,6 +378,7 @@ class UPythonGeneratedStruct : public UScriptStruct, public IPythonResourceOwner
 public:
 	//~ UObject interface
 	virtual void PostRename(UObject* OldOuter, const FName OldName) override;
+	virtual void BeginDestroy() override;
 
 	//~ UStruct interface
 	virtual void InitializeStruct(void* Dest, int32 ArrayDim = 1) const override;
diff --git a/Engine/Plugins/Experimental/PythonScriptPlugin/Source/PythonScriptPlugin/Private/PyWrapperTypeRegistry.cpp b/Engine/Plugins/Experimental/PythonScriptPlugin/Source/PythonScriptPlugin/Private/PyWrapperTypeRegistry.cpp
--- a/Engine/Plugins/Experimental/PythonScriptPlugin/Source/PythonScriptPlugin/Private/PyWrapperTypeRegistry.cpp
+++ b/Engine/Plugins/Experimental/PythonScriptPlugin/Source/PythonScriptPlugin/Private/PyWrapperTypeRegistry.cpp
@@ -2144,6 +2144,7 @@ void FPyWrapperTypeRegistry::GenerateStubCodeForWrappedTypes(const EPyOnlineDocs
 {
 	UE_LOG(LogPython, Display, TEXT("Generating Python API stub file..."));
 
+	FPyScopedGIL GIL;
 	FPyFileWriter PythonScript;
 
 	TUniquePtr<FPyOnlineDocsWriter> OnlineDocsWriter;
diff --git a/Engine/Plugins/Experimental/PythonScriptPlugin/Source/PythonScriptPlugin/Private/PythonScriptPlugin.cpp b/Engine/Plugins/Experimental/PythonScriptPlugin/Source/PythonScriptPlugin/Private/PythonScriptPlugin.cpp
--- a/Engine/Plugins/Experimental/PythonScriptPlugin/Source/PythonScriptPlugin/Private/PythonScriptPlugin.cpp
+++ b/Engine/Plugins/Experimental/PythonScriptPlugin/Source/PythonScriptPlugin/Private/PythonScriptPlugin.cpp
@@ -658,6 +658,13 @@ void FPythonScriptPlugin::InitializePython()
 		Py_SetPythonHome(PyHomePath.GetData());
 		Py_InitializeEx(0); // 0 so Python doesn't override any UE4 signal handling
 
+		// Ensure Python supports multiple threads via the GIL, as UE4 GC runs over multiple threads, 
+		// which may invoke FPyReferenceCollector::AddReferencedObjects on a background thread...
+		if (!PyEval_ThreadsInitialized())
+		{
+			PyEval_InitThreads();
+		}
+
 #if PLATFORM_WINDOWS && PY_MAJOR_VERSION >= 3
 		// We call _setmode here to restore the previous state
 		if (StdInMode != -1)
@@ -784,6 +791,10 @@ void FPythonScriptPlugin::InitializePython()
 		}));
 	}
 
+	// Release the GIL taken by Py_Initialize now that initialization has finished, to allow other threads access to Python
+	// We have to take this again prior to calling Py_Finalize, and all other code will lock on-demand via FPyScopedGIL
+	PyMainThreadState = PyEval_SaveThread();
+
 #if WITH_EDITOR
 	// Initialize the Content Browser integration
 	if (GIsEditor && !IsRunningCommandlet() && GetDefault<UPythonScriptPluginUserSettings>()->bEnableContentBrowserIntegration)
@@ -897,6 +908,10 @@ void FPythonScriptPlugin::ShutdownPython()
 		return;
 	}
 
+	// We need to restore the original GIL prior to calling Py_Finalize
+	PyEval_RestoreThread(PyMainThreadState);
+	PyMainThreadState = nullptr;
+
 #if WITH_EDITOR
 	// Remove the Content Browser integration
 	UToolMenus::UnregisterOwner(this);
diff --git a/Engine/Plugins/Experimental/PythonScriptPlugin/Source/PythonScriptPlugin/Private/PythonScriptPlugin.h b/Engine/Plugins/Experimental/PythonScriptPlugin/Source/PythonScriptPlugin/Private/PythonScriptPlugin.h
--- a/Engine/Plugins/Experimental/PythonScriptPlugin/Source/PythonScriptPlugin/Private/PythonScriptPlugin.h
+++ b/Engine/Plugins/Experimental/PythonScriptPlugin/Source/PythonScriptPlugin/Private/PythonScriptPlugin.h
@@ -205,6 +205,7 @@ private:
 	FPyObjectPtr PyConsoleGlobalDict;
 	FPyObjectPtr PyConsoleLocalDict;
 	FPyObjectPtr PyUnrealModule;
+	PyThreadState* PyMainThreadState = nullptr;
 	bool bInitialized;
 	bool bHasTicked;
 #endif	// WITH_PYTHON
diff --git a/Engine/Plugins/Experimental/PythonScriptPlugin/Source/PythonScriptPlugin/Private/PythonScriptRemoteExecution.cpp b/Engine/Plugins/Experimental/PythonScriptPlugin/Source/PythonScriptPlugin/Private/PythonScriptRemoteExecution.cpp
--- a/Engine/Plugins/Experimental/PythonScriptPlugin/Source/PythonScriptPlugin/Private/PythonScriptRemoteExecution.cpp
+++ b/Engine/Plugins/Experimental/PythonScriptPlugin/Source/PythonScriptPlugin/Private/PythonScriptRemoteExecution.cpp
@@ -297,12 +297,13 @@ public:
 				.BoundToAddress(MulticastBindAddress)
 				.BoundToPort(MulticastGroupEndpoint.Port)
 #if PLATFORM_SUPPORTS_UDP_MULTICAST_GROUP
-				.JoinedToGroup(MulticastGroupEndpoint.Address)
+				.JoinedToGroup(MulticastGroupEndpoint.Address, MulticastBindAddress)
 				.WithMulticastLoopback()
 #endif
 				.WithSendBufferSize(PythonScriptPluginSettings->RemoteExecutionSendBufferSizeBytes)
 				.WithReceiveBufferSize(PythonScriptPluginSettings->RemoteExecutionReceiveBufferSizeBytes)
-				.WithMulticastTtl(PythonScriptPluginSettings->RemoteExecutionMulticastTtl);
+				.WithMulticastTtl(PythonScriptPluginSettings->RemoteExecutionMulticastTtl)
+				.WithMulticastInterface(MulticastBindAddress);
 		}
 	}
 
