Ensure that blueprint actions are cleared for to-be reinstanced Python UClasses

diff --git a/Engine/Plugins/Experimental/PythonScriptPlugin/Source/PythonScriptPlugin/Private/PyWrapperObject.cpp b/Engine/Plugins/Experimental/PythonScriptPlugin/Source/PythonScriptPlugin/Private/PyWrapperObject.cpp
--- a/Engine/Plugins/Experimental/PythonScriptPlugin/Source/PythonScriptPlugin/Private/PyWrapperObject.cpp
+++ b/Engine/Plugins/Experimental/PythonScriptPlugin/Source/PythonScriptPlugin/Private/PyWrapperObject.cpp
@@ -1299,6 +1299,11 @@ public:
 		{
 			// Notify Blueprints that there is a new class to add to the action list
 			ActionDB->RefreshClassActions(NewClass);
+
+			if (OldClass)
+			{
+				ActionDB->RefreshClassActions(OldClass);
+			}
 		}
 
 		// Null the NewClass pointer so the destructor doesn't kill it
