Fixed Python shutdown crash when Slate callbacks are bound

Slate can shutdown prior to the Python plugin, which means it can clear an array of delegates (which may contain captured Python objects) without taking the Python GIL. TPyAutoGILPtr has been added as a way to capture Python instances that will safely take the GIL during copy, move, or destruction operations.

diff --git a/Engine/Plugins/Experimental/PythonScriptPlugin/Source/PythonScriptPlugin/Private/PyCore.cpp b/Engine/Plugins/Experimental/PythonScriptPlugin/Source/PythonScriptPlugin/Private/PyCore.cpp
--- a/Engine/Plugins/Experimental/PythonScriptPlugin/Source/PythonScriptPlugin/Private/PyCore.cpp
+++ b/Engine/Plugins/Experimental/PythonScriptPlugin/Source/PythonScriptPlugin/Private/PyCore.cpp
@@ -1,6 +1,7 @@
 // Copyright Epic Games, Inc. All Rights Reserved.
 
 #include "PyCore.h"
+#include "PyGIL.h"
 #include "PyGenUtil.h"
 #include "PyReferenceCollector.h"
 #include "PyWrapperTypeRegistry.h"
@@ -1729,10 +1730,9 @@ PyObject* RegisterPythonShutdownCallback(PyObject* InSelf, PyObject* InArgs)
 	}
 	check(PyObj);
 
-	FPyObjectPtr PyCallable = FPyObjectPtr::NewReference(PyObj);
-	FDelegateHandle PythonShutdownDelegateHandle = FPythonScriptPlugin::Get()->OnPythonShutdown().AddLambda([PyCallable]() mutable
+	FDelegateHandle PythonShutdownDelegateHandle = FPythonScriptPlugin::Get()->OnPythonShutdown().AddLambda([PyCallable = FPyAutoGILPtr(FPyObjectPtr::NewReference(PyObj))]() mutable
 	{
-		FPyObjectPtr Result = FPyObjectPtr::StealReference(PyObject_CallObject(PyCallable.GetPtr(), nullptr));
+		FPyObjectPtr Result = FPyObjectPtr::StealReference(PyObject_CallObject(PyCallable.Get().GetPtr(), nullptr));
 		if (!Result)
 		{
 			PyUtil::LogPythonError();
diff --git a/Engine/Plugins/Experimental/PythonScriptPlugin/Source/PythonScriptPlugin/Private/PyGIL.h b/Engine/Plugins/Experimental/PythonScriptPlugin/Source/PythonScriptPlugin/Private/PyGIL.h
--- a/Engine/Plugins/Experimental/PythonScriptPlugin/Source/PythonScriptPlugin/Private/PyGIL.h
+++ b/Engine/Plugins/Experimental/PythonScriptPlugin/Source/PythonScriptPlugin/Private/PyGIL.h
@@ -3,6 +3,7 @@
 #pragma once
 
 #include "IncludePython.h"
+#include "PyPtr.h"
 
 #if WITH_PYTHON
 
@@ -31,4 +32,105 @@ private:
 	PyGILState_STATE GILState;
 };
 
+/** Wrapper of a TPyPtr that can be safely copied/moved/destroyed by code where the GIL might not currently be held (eg, a lambda bound to a delegate) */
+template <typename TPythonType>
+class TPyAutoGILPtr
+{
+public:
+	TPyAutoGILPtr() = default;
+
+	explicit TPyAutoGILPtr(TPyPtr<TPythonType>&& InPtr)
+		: Ptr(MoveTemp(InPtr))
+	{
+	}
+
+	TPyAutoGILPtr(const TPyAutoGILPtr& InOther)
+	{
+		// This may be called after Python has already shut down
+		if (Py_IsInitialized())
+		{
+			FPyScopedGIL GIL;
+			Ptr = InOther.Ptr;
+		}
+	}
+
+	TPyAutoGILPtr(TPyAutoGILPtr&& InOther)
+	{
+		// This may be called after Python has already shut down
+		if (Py_IsInitialized())
+		{
+			FPyScopedGIL GIL;
+			Ptr = MoveTemp(InOther.Ptr);
+		}
+	}
+
+	~TPyAutoGILPtr()
+	{
+		// This may be called after Python has already shut down
+		if (Py_IsInitialized())
+		{
+			FPyScopedGIL GIL;
+			Ptr.Reset();
+		}
+		else
+		{
+			// Release ownership if Python has been shut down to avoid attempting to delete the object (which is already dead)
+			Ptr.Release();
+		}
+	}
+
+	TPyAutoGILPtr& operator=(const TPyAutoGILPtr& InOther)
+	{
+		if (this != &InOther)
+		{
+			// This may be called after Python has already shut down
+			if (Py_IsInitialized())
+			{
+				FPyScopedGIL GIL;
+				Ptr = InOther.Ptr;
+			}
+			else
+			{
+				// Release ownership if Python has been shut down to avoid attempting to delete the object (which is already dead)
+				Ptr.Release();
+			}
+		}
+		return *this;
+	}
+
+	TPyAutoGILPtr& operator=(TPyAutoGILPtr&& InOther)
+	{
+		if (this != &InOther)
+		{
+			// This may be called after Python has already shut down
+			if (Py_IsInitialized())
+			{
+				FPyScopedGIL GIL;
+				Ptr = MoveTemp(InOther.Ptr);
+			}
+			else
+			{
+				// Release ownership if Python has been shut down to avoid attempting to delete the object (which is already dead)
+				Ptr.Release();
+			}
+		}
+		return *this;
+	}
+
+	TPyPtr<TPythonType>& Get()
+	{
+		return Ptr;
+	}
+
+	const TPyPtr<TPythonType>& Get() const
+	{
+		return Ptr;
+	}
+
+private:
+	TPyPtr<TPythonType> Ptr;
+};
+
+typedef TPyAutoGILPtr<PyObject> FPyAutoGILPtr;
+
 #endif	// WITH_PYTHON
diff --git a/Engine/Plugins/Experimental/PythonScriptPlugin/Source/PythonScriptPlugin/Private/PySlate.cpp b/Engine/Plugins/Experimental/PythonScriptPlugin/Source/PythonScriptPlugin/Private/PySlate.cpp
--- a/Engine/Plugins/Experimental/PythonScriptPlugin/Source/PythonScriptPlugin/Private/PySlate.cpp
+++ b/Engine/Plugins/Experimental/PythonScriptPlugin/Source/PythonScriptPlugin/Private/PySlate.cpp
@@ -27,8 +27,7 @@ TArray<FDelegateHandle> PythonPostTickCallbackHandles;
 FPyDelegateHandle* RegisterSlateTickCallback(FSlateApplication::FSlateTickEvent& InSlateTickEvent, TArray<FDelegateHandle>& InOutPythonCallbackHandles, PyObject* InPyCallable)
 {
 	// We copy the PyCallable into the lambda to keep the Python object alive as long as the delegate is bound
-	FPyObjectPtr PyCallable = FPyObjectPtr::NewReference(InPyCallable);
-	FDelegateHandle TickEventDelegateHandle = InSlateTickEvent.AddLambda([PyCallable](const float InDeltaTime) mutable
+	FDelegateHandle TickEventDelegateHandle = InSlateTickEvent.AddLambda([PyCallable = FPyAutoGILPtr(FPyObjectPtr::NewReference(InPyCallable))](const float InDeltaTime) mutable
 	{
 		// Do not tick into Python when it may not be safe to call back into C++
 		if (GIsSavingPackage || IsGarbageCollecting() || FUObjectThreadContext::Get().IsRoutingPostLoad)
@@ -41,7 +40,7 @@ FPyDelegateHandle* RegisterSlateTickCallback(FSlateApplication::FSlateTickEvent&
 		FPyObjectPtr PyArgs = FPyObjectPtr::StealReference(PyTuple_New(1));
 		PyTuple_SetItem(PyArgs, 0, PyConversion::Pythonize(InDeltaTime)); // SetItem steals the reference
 
-		FPyObjectPtr Result = FPyObjectPtr::StealReference(PyObject_CallObject(PyCallable.GetPtr(), PyArgs));
+		FPyObjectPtr Result = FPyObjectPtr::StealReference(PyObject_CallObject(PyCallable.Get().GetPtr(), PyArgs));
 		if (!Result)
 		{
 			PyUtil::LogPythonError();
