Pre-cache the available on-disk Unreal Python modules

This can save 1-2s during Python initialization, by removing repeated calls to FileExists for everything on sys.path.

It does limit Unreal Python modules to being stored within the root of a known Content/Python directory, however this limitation is in-line with all existing modules that ship with the Engine, and allows us to avoid the cost of recursively scanning for modules.

diff --git a/Engine/Plugins/Experimental/PythonScriptPlugin/Source/PythonScriptPlugin/Private/PyUtil.cpp b/Engine/Plugins/Experimental/PythonScriptPlugin/Source/PythonScriptPlugin/Private/PyUtil.cpp
--- a/Engine/Plugins/Experimental/PythonScriptPlugin/Source/PythonScriptPlugin/Private/PyUtil.cpp
+++ b/Engine/Plugins/Experimental/PythonScriptPlugin/Source/PythonScriptPlugin/Private/PyUtil.cpp
@@ -18,9 +18,11 @@
 #include "PyWrapperTypeRegistry.h"
 
 #include "Misc/Paths.h"
+#include "Misc/PathViews.h"
 #include "Misc/ScopeExit.h"
 #include "Misc/MessageDialog.h"
 #include "Misc/DefaultValueHelper.h"
+#include "HAL/FileManager.h"
 #include "UObject/UnrealType.h"
 #include "UObject/EnumProperty.h"
 #include "UObject/TextProperty.h"
@@ -898,37 +900,105 @@ bool IsMappingType(PyTypeObject* InType)
 	return InType->tp_dict && PyDict_GetItemString(InType->tp_dict, "keys");
 }
 
-bool IsModuleAvailableForImport(const TCHAR* InModuleName, FString* OutResolvedFile)
+void FOnDiskModules::AddModules(const TCHAR* InPath)
+{
+	IFileManager& FileManager = IFileManager::Get();
+	FileManager.IterateDirectory(InPath, [this, &FileManager](const TCHAR* FilenameOrDirectory, bool bIsDirectory)
+	{
+		FString FullPath = FPaths::ConvertRelativePathToFull(FilenameOrDirectory);
+
+		bool bPassedWildcard = true;
+		if (!ModuleNameWildcard.IsEmpty())
+		{
+			// TODO: Would be nice to use FPathsView and FStringView here, but FStringView doesn't implement MatchesWildcard (though the implementation could easily be ported)
+			const FString CleanName = FPaths::GetCleanFilename(FullPath);
+			bPassedWildcard = CleanName.MatchesWildcard(ModuleNameWildcard, ESearchCase::CaseSensitive);
+		}
+
+		if (bPassedWildcard)
+		{
+			if (bIsDirectory)
+			{
+				FullPath /= TEXT("__init__.py");
+				if (FileManager.FileExists(*FullPath))
+				{
+					CachedModules.Add(MoveTemp(FullPath));
+				}
+			}
+			else if (FPathViews::GetExtension(FullPath) == TEXT("py"))
+			{
+				CachedModules.Add(MoveTemp(FullPath));
+			}
+		}
+
+		return true;
+	});
+}
+
+void FOnDiskModules::RemoveModules(const TCHAR* InPath)
+{
+	for (auto It = CachedModules.CreateIterator(); It; ++It)
+	{
+		if (It->StartsWith(InPath))
+		{
+			It.RemoveCurrent();
+		}
+	}
+}
+
+bool FOnDiskModules::HasModule(const TCHAR* InModuleName, FString* OutResolvedFile) const
+{
+	const FString ModuleSingleFile = FString::Printf(TEXT("/%s.py"), InModuleName);
+	const FString ModuleFolderName = FString::Printf(TEXT("/%s/__init__.py"), InModuleName);
+
+	for (const FString& CachedModule : CachedModules)
+	{
+		if (CachedModule.EndsWith(ModuleSingleFile, ESearchCase::CaseSensitive) || CachedModule.EndsWith(ModuleFolderName, ESearchCase::CaseSensitive))
+		{
+			if (OutResolvedFile)
+			{
+				*OutResolvedFile = CachedModule;
+			}
+
+			return true;
+		}
+	}
+
+	return false;
+}
+
+FOnDiskModules& GetOnDiskUnrealModulesCache()
+{
+	static FOnDiskModules OnDiskUnrealModules(TEXT("unreal_*"));
+	return OnDiskUnrealModules;
+}
+
+bool IsModuleAvailableForImport(const TCHAR* InModuleName, const FOnDiskModules* InOnDiskModules, FString* OutResolvedFile)
 {
 	// Check the sys.modules table first since it avoids hitting the filesystem
 	if (PyObject* PyModulesDict = PySys_GetObject(PyCStrCast("modules")))
 	{
-		PyObject* PyModuleKey = nullptr;
-		PyObject* PyModuleValue = nullptr;
-		Py_ssize_t ModuleDictIndex = 0;
-		while (PyDict_Next(PyModulesDict, &ModuleDictIndex, &PyModuleKey, &PyModuleValue))
+		if (PyObject* PyModuleValue = PyDict_GetItemString(PyModulesDict, TCHAR_TO_UTF8(InModuleName)))
 		{
-			if (PyModuleKey)
+			if (OutResolvedFile)
 			{
-				const FString CurModuleName = PyObjectToUEString(PyModuleKey);
-				if (FCString::Strcmp(InModuleName, *CurModuleName) == 0)
+				PyObject* PyModuleDict = PyModule_GetDict(PyModuleValue);
+				if (PyObject* PyModuleFile = PyDict_GetItemString(PyModuleDict, "__file__"))
 				{
-					if (OutResolvedFile && PyModuleValue)
-					{
-						PyObject* PyModuleDict = PyModule_GetDict(PyModuleValue);
-						PyObject* PyModuleFile = PyDict_GetItemString(PyModuleDict, "__file__");
-						if (PyModuleFile)
-						{
-							*OutResolvedFile = PyObjectToUEString(PyModuleFile);
-						}
-					}
-
-					return true;
+					*OutResolvedFile = PyObjectToUEString(PyModuleFile);
 				}
 			}
+
+			return true;
 		}
 	}
 
+	// Use the on-disk modules cache, if available, to avoid hitting the filesystem
+	if (InOnDiskModules)
+	{
+		return InOnDiskModules->HasModule(InModuleName, OutResolvedFile);
+	}
+
 	// Check the sys.path list looking for bla.py or bla/__init__.py
 	if (PyObject* PyPathList = PySys_GetObject(PyCStrCast("path")))
 	{
@@ -973,23 +1043,14 @@ bool IsModuleImported(const TCHAR* InModuleName, PyObject** OutPyModule)
 {
 	if (PyObject* PyModulesDict = PySys_GetObject(PyCStrCast("modules")))
 	{
-		PyObject* PyModuleKey = nullptr;
-		PyObject* PyModuleValue = nullptr;
-		Py_ssize_t ModuleDictIndex = 0;
-		while (PyDict_Next(PyModulesDict, &ModuleDictIndex, &PyModuleKey, &PyModuleValue))
+		if (PyObject* PyModuleValue = PyDict_GetItemString(PyModulesDict, TCHAR_TO_UTF8(InModuleName)))
 		{
-			if (PyModuleKey)
+			if (OutPyModule)
 			{
-				const FString CurModuleName = PyObjectToUEString(PyModuleKey);
-				if (FCString::Strcmp(InModuleName, *CurModuleName) == 0)
-				{
-					if (OutPyModule)
-					{
-						*OutPyModule = PyModuleValue;
-					}
-					return true;
-				}
+				*OutPyModule = PyModuleValue;
 			}
+
+			return true;
 		}
 	}
 
diff --git a/Engine/Plugins/Experimental/PythonScriptPlugin/Source/PythonScriptPlugin/Private/PyUtil.h b/Engine/Plugins/Experimental/PythonScriptPlugin/Source/PythonScriptPlugin/Private/PyUtil.h
--- a/Engine/Plugins/Experimental/PythonScriptPlugin/Source/PythonScriptPlugin/Private/PyUtil.h
+++ b/Engine/Plugins/Experimental/PythonScriptPlugin/Source/PythonScriptPlugin/Private/PyUtil.h
@@ -408,11 +408,53 @@ namespace PyUtil
 	 */
 	bool IsMappingType(PyTypeObject* InType);
 
+	/**
+	 * Cache of on-disk Python modules to optimize repeated IsModuleAvailableForImport queries.
+	 * @note This cache does not automatically watch or update itself from file-system changes after the initial scan!
+	 */
+	class FOnDiskModules
+	{
+	public:
+		FOnDiskModules() = default;
+
+		explicit FOnDiskModules(FString InModuleNameWildcard)
+			: ModuleNameWildcard(MoveTemp(InModuleNameWildcard))
+		{
+		}
+
+		/**
+		 * Add any module files found under the given path.
+		 * @note Function is non-recursive, and will find both "{Module}.py" and "{Module}/__init__.py" module layouts.
+		 */
+		void AddModules(const TCHAR* InPath);
+
+		/**
+		 * Remove any module files found under the given path.
+		 */
+		void RemoveModules(const TCHAR* InPath);
+
+		/**
+		 * Check to see whether the given module is known to the cache, optionally returning the actual module file location if it is.
+		 * @note Will look for both "{Module}.py" and "{Module}/__init__.py" module layouts.
+		 */
+		bool HasModule(const TCHAR* InModuleName, FString* OutResolvedFile = nullptr) const;
+
+	private:
+		FString ModuleNameWildcard;
+
+		TSet<FString> CachedModules;
+	};
+
+	/**
+	 * Get the cache of on-disk Unreal Python modules ("unreal_*" wildcard).
+	 */
+	FOnDiskModules& GetOnDiskUnrealModulesCache();
+
 	/**
 	 * Test to see whether the given module is available for import (is in the sys.modules table (which includes built-in modules), or is in any known sys.path path).
 	 * @note This function can't handle dot separated names.
 	 */
-	bool IsModuleAvailableForImport(const TCHAR* InModuleName, FString* OutResolvedFile = nullptr);
+	bool IsModuleAvailableForImport(const TCHAR* InModuleName, const FOnDiskModules* InOnDiskModules = nullptr, FString* OutResolvedFile = nullptr);
 
 	/**
 	 * Test to see whether the given module is available for import (is in the sys.modules table).
diff --git a/Engine/Plugins/Experimental/PythonScriptPlugin/Source/PythonScriptPlugin/Private/PyWrapperTypeRegistry.cpp b/Engine/Plugins/Experimental/PythonScriptPlugin/Source/PythonScriptPlugin/Private/PyWrapperTypeRegistry.cpp
--- a/Engine/Plugins/Experimental/PythonScriptPlugin/Source/PythonScriptPlugin/Private/PyWrapperTypeRegistry.cpp
+++ b/Engine/Plugins/Experimental/PythonScriptPlugin/Source/PythonScriptPlugin/Private/PyWrapperTypeRegistry.cpp
@@ -2182,7 +2182,7 @@ void FPyWrapperTypeRegistry::GenerateStubCodeForWrappedTypes(const EPyOnlineDocs
 			const FString PythonModuleName = FString::Printf(TEXT("unreal_%s"), *PythonBaseModuleName);
 		
 			FString ModuleFilename;
-			if (PyUtil::IsModuleAvailableForImport(*PythonModuleName, &ModuleFilename))
+			if (PyUtil::IsModuleAvailableForImport(*PythonModuleName, &PyUtil::GetOnDiskUnrealModulesCache(), &ModuleFilename))
 			{
 				// Adjust .pyc and .pyd files so we try and find the source Python file
 				ModuleFilename = FPaths::ChangeExtension(ModuleFilename, TEXT(".py"));
diff --git a/Engine/Plugins/Experimental/PythonScriptPlugin/Source/PythonScriptPlugin/Private/PythonScriptPlugin.cpp b/Engine/Plugins/Experimental/PythonScriptPlugin/Source/PythonScriptPlugin/Private/PythonScriptPlugin.cpp
--- a/Engine/Plugins/Experimental/PythonScriptPlugin/Source/PythonScriptPlugin/Private/PythonScriptPlugin.cpp
+++ b/Engine/Plugins/Experimental/PythonScriptPlugin/Source/PythonScriptPlugin/Private/PythonScriptPlugin.cpp
@@ -737,7 +737,9 @@ void FPythonScriptPlugin::InitializePython()
 		for (const FString& RootPath : RootPaths)
 		{
 			const FString RootFilesystemPath = FPackageName::LongPackageNameToFilename(RootPath);
-			PyUtil::AddSystemPath(FPaths::ConvertRelativePathToFull(RootFilesystemPath / TEXT("Python")));
+			const FString PythonContentPath = FPaths::ConvertRelativePathToFull(RootFilesystemPath / TEXT("Python"));
+			PyUtil::AddSystemPath(PythonContentPath);
+			PyUtil::GetOnDiskUnrealModulesCache().AddModules(*PythonContentPath);
 		}
 
 		for (const FDirectoryPath& AdditionalPath : GetDefault<UPythonScriptPluginSettings>()->AdditionalPaths)
@@ -1081,7 +1083,7 @@ void FPythonScriptPlugin::ImportUnrealModule(const TCHAR* InModuleName)
 
 	const TCHAR* ModuleNameToImport = nullptr;
 	PyObject* ModuleToReload = nullptr;
-	if (PyUtil::IsModuleAvailableForImport(*PythonModuleName))
+	if (PyUtil::IsModuleAvailableForImport(*PythonModuleName, &PyUtil::GetOnDiskUnrealModulesCache()))
 	{
 		// Python modules that are already loaded should be reloaded if we're requested to import them again
 		if (!PyUtil::IsModuleImported(*PythonModuleName, &ModuleToReload))
@@ -1089,7 +1091,7 @@ void FPythonScriptPlugin::ImportUnrealModule(const TCHAR* InModuleName)
 			ModuleNameToImport = *PythonModuleName;
 		}
 	}
-	else if (PyUtil::IsModuleAvailableForImport(*NativeModuleName))
+	else if (PyUtil::IsModuleAvailableForImport(*NativeModuleName, &PyUtil::GetOnDiskUnrealModulesCache()))
 	{
 		ModuleNameToImport = *NativeModuleName;
 	}
@@ -1348,7 +1350,9 @@ void FPythonScriptPlugin::OnContentPathMounted(const FString& InAssetPath, const
 {
 	{
 		FPyScopedGIL GIL;
-		PyUtil::AddSystemPath(FPaths::ConvertRelativePathToFull(InFilesystemPath / TEXT("Python")));
+		const FString PythonContentPath = FPaths::ConvertRelativePathToFull(InFilesystemPath / TEXT("Python"));
+		PyUtil::AddSystemPath(PythonContentPath);
+		PyUtil::GetOnDiskUnrealModulesCache().AddModules(*PythonContentPath);
 	}
 
 #if WITH_EDITOR
@@ -1363,7 +1367,9 @@ void FPythonScriptPlugin::OnContentPathDismounted(const FString& InAssetPath, co
 {
 	{
 		FPyScopedGIL GIL;
-		PyUtil::RemoveSystemPath(FPaths::ConvertRelativePathToFull(InFilesystemPath / TEXT("Python")));
+		const FString PythonContentPath = FPaths::ConvertRelativePathToFull(InFilesystemPath / TEXT("Python"));
+		PyUtil::RemoveSystemPath(PythonContentPath);
+		PyUtil::GetOnDiskUnrealModulesCache().RemoveModules(*PythonContentPath);
 	}
 
 #if WITH_EDITOR
