Renamed literal UE4 references in PythonScriptPlugin

#rnx

diff --git a/Engine/Plugins/Experimental/PythonScriptPlugin/Source/PythonScriptPlugin/Private/PyCore.cpp b/Engine/Plugins/Experimental/PythonScriptPlugin/Source/PythonScriptPlugin/Private/PyCore.cpp
--- a/Engine/Plugins/Experimental/PythonScriptPlugin/Source/PythonScriptPlugin/Private/PyCore.cpp
+++ b/Engine/Plugins/Experimental/PythonScriptPlugin/Source/PythonScriptPlugin/Private/PyCore.cpp
@@ -995,7 +995,7 @@ PyTypeObject InitializePyUFunctionDefType()
 }
 
 
-PyTypeObject PyDelegateHandleType = InitializePyWrapperBasicType<FPyDelegateHandle>("_DelegateHandle", "Type for all UE4 exposed FDelegateHandle instances");
+PyTypeObject PyDelegateHandleType = InitializePyWrapperBasicType<FPyDelegateHandle>("_DelegateHandle", "Type for all Unreal exposed FDelegateHandle instances");
 PyTypeObject PyScopedSlowTaskType = InitializePyScopedSlowTaskType();
 PyTypeObject PyObjectIteratorType = InitializePyObjectIteratorType();
 PyTypeObject PyClassIteratorType = InitializePyTypeIteratorType<UClass, FPyClassIterator>("ClassIterator", "Type for iterating Unreal class types");
diff --git a/Engine/Plugins/Experimental/PythonScriptPlugin/Source/PythonScriptPlugin/Private/PyCore.h b/Engine/Plugins/Experimental/PythonScriptPlugin/Source/PythonScriptPlugin/Private/PyCore.h
--- a/Engine/Plugins/Experimental/PythonScriptPlugin/Source/PythonScriptPlugin/Private/PyCore.h
+++ b/Engine/Plugins/Experimental/PythonScriptPlugin/Source/PythonScriptPlugin/Private/PyCore.h
@@ -48,7 +48,7 @@ extern PyTypeObject PyFPropertyDefType;
 /** Python type for FPyUFunctionDef */
 extern PyTypeObject PyUFunctionDefType;
 
-/** Type for all UE4 exposed FDelegateHandle instances */
+/** Type for all Unreal exposed FDelegateHandle instances */
 struct FPyDelegateHandle : public TPyWrapperBasic<FDelegateHandle, FPyDelegateHandle>
 {
 	typedef TPyWrapperBasic<FDelegateHandle, FPyDelegateHandle> Super;
diff --git a/Engine/Plugins/Experimental/PythonScriptPlugin/Source/PythonScriptPlugin/Private/PyWrapperArray.cpp b/Engine/Plugins/Experimental/PythonScriptPlugin/Source/PythonScriptPlugin/Private/PyWrapperArray.cpp
--- a/Engine/Plugins/Experimental/PythonScriptPlugin/Source/PythonScriptPlugin/Private/PyWrapperArray.cpp
+++ b/Engine/Plugins/Experimental/PythonScriptPlugin/Source/PythonScriptPlugin/Private/PyWrapperArray.cpp
@@ -1445,7 +1445,7 @@ PyTypeObject InitializePyWrapperArrayType()
 	PyType.tp_methods = PyMethods;
 
 	PyType.tp_flags = Py_TPFLAGS_DEFAULT;
-	PyType.tp_doc = "Type for all UE4 exposed array instances";
+	PyType.tp_doc = "Type for all Unreal exposed array instances";
 
 	static PySequenceMethods PySequence;
 	PySequence.sq_length = (lenfunc)&FSequenceFuncs::Len;
@@ -1522,7 +1522,7 @@ PyTypeObject InitializePyWrapperArrayIteratorType()
 	PyType.tp_iternext = (iternextfunc)&FFuncs::IterNext;
 
 	PyType.tp_flags = Py_TPFLAGS_DEFAULT;
-	PyType.tp_doc = "Type for all UE4 exposed array iterators";
+	PyType.tp_doc = "Type for all Unreal exposed array iterators";
 
 	return PyType;
 }
diff --git a/Engine/Plugins/Experimental/PythonScriptPlugin/Source/PythonScriptPlugin/Private/PyWrapperArray.h b/Engine/Plugins/Experimental/PythonScriptPlugin/Source/PythonScriptPlugin/Private/PyWrapperArray.h
--- a/Engine/Plugins/Experimental/PythonScriptPlugin/Source/PythonScriptPlugin/Private/PyWrapperArray.h
+++ b/Engine/Plugins/Experimental/PythonScriptPlugin/Source/PythonScriptPlugin/Private/PyWrapperArray.h
@@ -14,7 +14,7 @@ extern PyTypeObject PyWrapperArrayType;
 /** Initialize the PyWrapperArray types and add them to the given Python module */
 void InitializePyWrapperArray(PyGenUtil::FNativePythonModule& ModuleInfo);
 
-/** Type for all UE4 exposed array instances */
+/** Type for all Unreal exposed array instances */
 struct FPyWrapperArray : public FPyWrapperBase
 {
 	/** The owner of the wrapped array instance (if any) */
@@ -106,7 +106,7 @@ struct FPyWrapperArray : public FPyWrapperBase
 	static int Resize(FPyWrapperArray* InSelf, Py_ssize_t InLen);
 };
 
-/** Meta-data for all UE4 exposed array types */
+/** Meta-data for all Unreal exposed array types */
 struct FPyWrapperArrayMetaData : public FPyWrapperBaseMetaData
 {
 	PY_METADATA_METHODS(FPyWrapperArrayMetaData, FGuid(0xFCD8DF06, 0x7AD6426F, 0x97482D82, 0x6E100C39))
diff --git a/Engine/Plugins/Experimental/PythonScriptPlugin/Source/PythonScriptPlugin/Private/PyWrapperBase.cpp b/Engine/Plugins/Experimental/PythonScriptPlugin/Source/PythonScriptPlugin/Private/PyWrapperBase.cpp
--- a/Engine/Plugins/Experimental/PythonScriptPlugin/Source/PythonScriptPlugin/Private/PyWrapperBase.cpp
+++ b/Engine/Plugins/Experimental/PythonScriptPlugin/Source/PythonScriptPlugin/Private/PyWrapperBase.cpp
@@ -87,7 +87,7 @@ PyTypeObject InitializePyWrapperBaseType()
 	PyType.tp_init = (initproc)&FFuncs::Init;
 
 	PyType.tp_flags = Py_TPFLAGS_DEFAULT;
-	PyType.tp_doc = "Base type for all UE4 exposed types";
+	PyType.tp_doc = "Base type for all Unreal exposed types";
 
 	return PyType;
 }
diff --git a/Engine/Plugins/Experimental/PythonScriptPlugin/Source/PythonScriptPlugin/Private/PyWrapperBase.h b/Engine/Plugins/Experimental/PythonScriptPlugin/Source/PythonScriptPlugin/Private/PyWrapperBase.h
--- a/Engine/Plugins/Experimental/PythonScriptPlugin/Source/PythonScriptPlugin/Private/PyWrapperBase.h
+++ b/Engine/Plugins/Experimental/PythonScriptPlugin/Source/PythonScriptPlugin/Private/PyWrapperBase.h
@@ -20,7 +20,7 @@ extern PyTypeObject PyWrapperBaseType;
 /** Initialize the PyWrapperBase types and add them to the given Python module */
 void InitializePyWrapperBase(PyGenUtil::FNativePythonModule& ModuleInfo);
 
-/** Base type for all UE4 exposed instances */
+/** Base type for all Unreal exposed instances */
 struct FPyWrapperBase
 {
 	/** Common Python Object */
@@ -49,7 +49,7 @@ struct FPyWrapperBase
 	static FGuid StaticTypeId() { return (GUID); }																				\
 	virtual FGuid GetTypeId() const override { return StaticTypeId(); }
 
-/** Base meta-data for all UE4 exposed types */
+/** Base meta-data for all Unreal exposed types */
 struct FPyWrapperBaseMetaData
 {
 	/** Set the meta-data object on the given type */
diff --git a/Engine/Plugins/Experimental/PythonScriptPlugin/Source/PythonScriptPlugin/Private/PyWrapperBasic.h b/Engine/Plugins/Experimental/PythonScriptPlugin/Source/PythonScriptPlugin/Private/PyWrapperBasic.h
--- a/Engine/Plugins/Experimental/PythonScriptPlugin/Source/PythonScriptPlugin/Private/PyWrapperBasic.h
+++ b/Engine/Plugins/Experimental/PythonScriptPlugin/Source/PythonScriptPlugin/Private/PyWrapperBasic.h
@@ -44,7 +44,7 @@ PyTypeObject InitializePyWrapperBasicType(const char* InTypeName, const char* In
 	return PyType;
 }
 
-/** Base type for any UE4 exposed simple value instances (that just copy data into Python) */
+/** Base type for any Unreal exposed simple value instances (that just copy data into Python) */
 template <typename ValueType, typename SelfType>
 struct TPyWrapperBasic : public FPyWrapperBase
 {
diff --git a/Engine/Plugins/Experimental/PythonScriptPlugin/Source/PythonScriptPlugin/Private/PyWrapperDelegate.cpp b/Engine/Plugins/Experimental/PythonScriptPlugin/Source/PythonScriptPlugin/Private/PyWrapperDelegate.cpp
--- a/Engine/Plugins/Experimental/PythonScriptPlugin/Source/PythonScriptPlugin/Private/PyWrapperDelegate.cpp
+++ b/Engine/Plugins/Experimental/PythonScriptPlugin/Source/PythonScriptPlugin/Private/PyWrapperDelegate.cpp
@@ -779,7 +779,7 @@ PyTypeObject InitializePyWrapperDelegateType()
 	PyType.tp_methods = PyMethods;
 
 	PyType.tp_flags = Py_TPFLAGS_DEFAULT;
-	PyType.tp_doc = "Type for all UE4 exposed delegate instances";
+	PyType.tp_doc = "Type for all Unreal exposed delegate instances";
 
 	static PyNumberMethods PyNumber;
 #if PY_MAJOR_VERSION >= 3
@@ -1207,7 +1207,7 @@ PyTypeObject InitializePyWrapperMulticastDelegateType()
 	PyType.tp_methods = PyMethods;
 
 	PyType.tp_flags = Py_TPFLAGS_DEFAULT;
-	PyType.tp_doc = "Type for all UE4 exposed multicast delegate instances";
+	PyType.tp_doc = "Type for all Unreal exposed multicast delegate instances";
 
 	static PyNumberMethods PyNumber;
 #if PY_MAJOR_VERSION >= 3
diff --git a/Engine/Plugins/Experimental/PythonScriptPlugin/Source/PythonScriptPlugin/Private/PyWrapperDelegate.h b/Engine/Plugins/Experimental/PythonScriptPlugin/Source/PythonScriptPlugin/Private/PyWrapperDelegate.h
--- a/Engine/Plugins/Experimental/PythonScriptPlugin/Source/PythonScriptPlugin/Private/PyWrapperDelegate.h
+++ b/Engine/Plugins/Experimental/PythonScriptPlugin/Source/PythonScriptPlugin/Private/PyWrapperDelegate.h
@@ -63,7 +63,7 @@ extern PyTypeObject PyWrapperMulticastDelegateType;
 /** Initialize the PyWrapperDelegate types and add them to the given Python module */
 void InitializePyWrapperDelegate(PyGenUtil::FNativePythonModule& ModuleInfo);
 
-/** Base type for all UE4 exposed delegate instances */
+/** Base type for all Unreal exposed delegate instances */
 template <typename DelegateType>
 struct TPyWrapperDelegate : public FPyWrapperBase
 {
@@ -77,7 +77,7 @@ struct TPyWrapperDelegate : public FPyWrapperBase
 	DelegateType InternalDelegateInstance;
 };
 
-/** Base meta-data for all UE4 exposed delegate types */
+/** Base meta-data for all Unreal exposed delegate types */
 template <typename WrapperType>
 struct TPyWrapperDelegateMetaData : public FPyWrapperBaseMetaData
 {
@@ -135,7 +135,7 @@ struct TPyWrapperDelegateMetaData : public FPyWrapperBaseMetaData
 	const UClass* PythonCallableForDelegateClass;
 };
 
-/** Type for all UE4 exposed delegate instances */
+/** Type for all Unreal exposed delegate instances */
 struct FPyWrapperDelegate : public TPyWrapperDelegate<FScriptDelegate>
 {
 	/** New this wrapper instance (called via tp_new for Python, or directly in C++) */
@@ -166,7 +166,7 @@ struct FPyWrapperDelegate : public TPyWrapperDelegate<FScriptDelegate>
 	static PyObject* CallDelegate(FPyWrapperDelegate* InSelf, PyObject* InArgs);
 };
 
-/** Meta-data for all UE4 exposed delegate types */
+/** Meta-data for all Unreal exposed delegate types */
 struct FPyWrapperDelegateMetaData : public TPyWrapperDelegateMetaData<FPyWrapperDelegate>
 {
 	PY_METADATA_METHODS(FPyWrapperDelegateMetaData, FGuid(0xCB3D0485, 0x8A3A443E, 0xBEE336F4, 0x82888A81))
@@ -175,7 +175,7 @@ struct FPyWrapperDelegateMetaData : public TPyWrapperDelegateMetaData<FPyWrapper
 	virtual void AddReferencedObjects(FPyWrapperBase* Instance, FReferenceCollector& Collector) override;
 };
 
-/** Type for all UE4 exposed multicast delegate instances */
+/** Type for all Unreal exposed multicast delegate instances */
 struct FPyWrapperMulticastDelegate : public TPyWrapperDelegate<FMulticastScriptDelegate>
 {
 	/** New this wrapper instance (called via tp_new for Python, or directly in C++) */
@@ -206,7 +206,7 @@ struct FPyWrapperMulticastDelegate : public TPyWrapperDelegate<FMulticastScriptD
 	static PyObject* CallDelegate(FPyWrapperMulticastDelegate* InSelf, PyObject* InArgs);
 };
 
-/** Meta-data for all UE4 exposed multicast delegate types */
+/** Meta-data for all Unreal exposed multicast delegate types */
 struct FPyWrapperMulticastDelegateMetaData : public TPyWrapperDelegateMetaData<FPyWrapperMulticastDelegate>
 {
 	PY_METADATA_METHODS(FPyWrapperMulticastDelegateMetaData, FGuid(0x448FB4DA, 0x38DC4386, 0xBCAFF448, 0x29C0F3A4))
diff --git a/Engine/Plugins/Experimental/PythonScriptPlugin/Source/PythonScriptPlugin/Private/PyWrapperEnum.cpp b/Engine/Plugins/Experimental/PythonScriptPlugin/Source/PythonScriptPlugin/Private/PyWrapperEnum.cpp
--- a/Engine/Plugins/Experimental/PythonScriptPlugin/Source/PythonScriptPlugin/Private/PyWrapperEnum.cpp
+++ b/Engine/Plugins/Experimental/PythonScriptPlugin/Source/PythonScriptPlugin/Private/PyWrapperEnum.cpp
@@ -435,7 +435,7 @@ PyTypeObject InitializePyWrapperEnumType()
 	PyType.tp_methods = PyMethods;
 
 	PyType.tp_flags = Py_TPFLAGS_DEFAULT | Py_TPFLAGS_BASETYPE;
-	PyType.tp_doc = "Type for all UE4 exposed enum instances";
+	PyType.tp_doc = "Type for all Unreal exposed enum instances";
 
 	return PyType;
 }
@@ -580,7 +580,7 @@ PyTypeObject InitializePyWrapperEnumMetaclassType()
 	PyType.tp_iter = (getiterfunc)&FFuncs::GetIter;
 
 	PyType.tp_flags = Py_TPFLAGS_DEFAULT | Py_TPFLAGS_BASETYPE;
-	PyType.tp_doc = "Metaclass type for all UE4 exposed enum instances";
+	PyType.tp_doc = "Metaclass type for all Unreal exposed enum instances";
 
 	static PySequenceMethods PySequence;
 	PySequence.sq_length = (lenfunc)&FSequenceFuncs::Len;
@@ -651,7 +651,7 @@ PyTypeObject InitializePyWrapperEnumIteratorType()
 	PyType.tp_iternext = (iternextfunc)&FFuncs::IterNext;
 
 	PyType.tp_flags = Py_TPFLAGS_DEFAULT;
-	PyType.tp_doc = "Type for all UE4 exposed enum iterators";
+	PyType.tp_doc = "Type for all Unreal exposed enum iterators";
 
 	return PyType;
 }
diff --git a/Engine/Plugins/Experimental/PythonScriptPlugin/Source/PythonScriptPlugin/Private/PyWrapperEnum.h b/Engine/Plugins/Experimental/PythonScriptPlugin/Source/PythonScriptPlugin/Private/PyWrapperEnum.h
--- a/Engine/Plugins/Experimental/PythonScriptPlugin/Source/PythonScriptPlugin/Private/PyWrapperEnum.h
+++ b/Engine/Plugins/Experimental/PythonScriptPlugin/Source/PythonScriptPlugin/Private/PyWrapperEnum.h
@@ -17,7 +17,7 @@ extern PyTypeObject PyWrapperEnumValueDescrType;
 /** Initialize the PyWrapperEnum types and add them to the given Python module */
 void InitializePyWrapperEnum(PyGenUtil::FNativePythonModule& ModuleInfo);
 
-/** Type for all UE4 exposed enum instances (an instance is created for each entry in the enum, before the enum type is locked for creating new instances) */
+/** Type for all Unreal exposed enum instances (an instance is created for each entry in the enum, before the enum type is locked for creating new instances) */
 struct FPyWrapperEnum : public FPyWrapperBase
 {
 	/** Name of this enum entry */
@@ -60,7 +60,7 @@ struct FPyWrapperEnum : public FPyWrapperBase
 	static FPyWrapperEnum* AddEnumEntry(PyTypeObject* InType, const int64 InEnumEntryValue, const char* InEnumEntryName, const char* InEnumEntryDoc);
 };
 
-/** Meta-data for all UE4 exposed enum types */
+/** Meta-data for all Unreal exposed enum types */
 struct FPyWrapperEnumMetaData : public FPyWrapperBaseMetaData
 {
 	PY_METADATA_METHODS(FPyWrapperEnumMetaData, FGuid(0x1D69987C, 0x2F624403, 0x8379FCB5, 0xF896B595))
diff --git a/Engine/Plugins/Experimental/PythonScriptPlugin/Source/PythonScriptPlugin/Private/PyWrapperFixedArray.cpp b/Engine/Plugins/Experimental/PythonScriptPlugin/Source/PythonScriptPlugin/Private/PyWrapperFixedArray.cpp
--- a/Engine/Plugins/Experimental/PythonScriptPlugin/Source/PythonScriptPlugin/Private/PyWrapperFixedArray.cpp
+++ b/Engine/Plugins/Experimental/PythonScriptPlugin/Source/PythonScriptPlugin/Private/PyWrapperFixedArray.cpp
@@ -819,7 +819,7 @@ PyTypeObject InitializePyWrapperFixedArrayType()
 	PyType.tp_methods = PyMethods;
 
 	PyType.tp_flags = Py_TPFLAGS_DEFAULT;
-	PyType.tp_doc = "Type for all UE4 exposed fixed-array instances";
+	PyType.tp_doc = "Type for all Unreal exposed fixed-array instances";
 
 	static PySequenceMethods PySequence;
 	PySequence.sq_length = (lenfunc)&FSequenceFuncs::Len;
@@ -893,7 +893,7 @@ PyTypeObject InitializePyWrapperFixedArrayIteratorType()
 	PyType.tp_iternext = (iternextfunc)&FFuncs::IterNext;
 
 	PyType.tp_flags = Py_TPFLAGS_DEFAULT;
-	PyType.tp_doc = "Type for all UE4 exposed fixed-array iterators";
+	PyType.tp_doc = "Type for all Unreal exposed fixed-array iterators";
 
 	return PyType;
 }
diff --git a/Engine/Plugins/Experimental/PythonScriptPlugin/Source/PythonScriptPlugin/Private/PyWrapperFixedArray.h b/Engine/Plugins/Experimental/PythonScriptPlugin/Source/PythonScriptPlugin/Private/PyWrapperFixedArray.h
--- a/Engine/Plugins/Experimental/PythonScriptPlugin/Source/PythonScriptPlugin/Private/PyWrapperFixedArray.h
+++ b/Engine/Plugins/Experimental/PythonScriptPlugin/Source/PythonScriptPlugin/Private/PyWrapperFixedArray.h
@@ -14,7 +14,7 @@ extern PyTypeObject PyWrapperFixedArrayType;
 /** Initialize the PyWrapperFixedArray types and add them to the given Python module */
 void InitializePyWrapperFixedArray(PyGenUtil::FNativePythonModule& ModuleInfo);
 
-/** Type for all UE4 exposed fixed-array instances */
+/** Type for all Unreal exposed fixed-array instances */
 struct FPyWrapperFixedArray : public FPyWrapperBase
 {
 	/** The owner of the wrapped fixed-array instance (if any) */
@@ -72,7 +72,7 @@ struct FPyWrapperFixedArray : public FPyWrapperBase
 	static FPyWrapperFixedArray* Repeat(FPyWrapperFixedArray* InSelf, Py_ssize_t InMultiplier);
 };
 
-/** Meta-data for all UE4 exposed fixed-array types */
+/** Meta-data for all Unreal exposed fixed-array types */
 struct FPyWrapperFixedArrayMetaData : public FPyWrapperBaseMetaData
 {
 	PY_METADATA_METHODS(FPyWrapperFixedArrayMetaData, FGuid(0x8A8CA89E, 0xC8F04627, 0xA7D2B568, 0xB086C02E))
diff --git a/Engine/Plugins/Experimental/PythonScriptPlugin/Source/PythonScriptPlugin/Private/PyWrapperMap.cpp b/Engine/Plugins/Experimental/PythonScriptPlugin/Source/PythonScriptPlugin/Private/PyWrapperMap.cpp
--- a/Engine/Plugins/Experimental/PythonScriptPlugin/Source/PythonScriptPlugin/Private/PyWrapperMap.cpp
+++ b/Engine/Plugins/Experimental/PythonScriptPlugin/Source/PythonScriptPlugin/Private/PyWrapperMap.cpp
@@ -1611,7 +1611,7 @@ PyTypeObject InitializePyWrapperMapType()
 	PyType.tp_methods = PyMethods;
 
 	PyType.tp_flags = Py_TPFLAGS_DEFAULT;
-	PyType.tp_doc = "Type for all UE4 exposed map instances";
+	PyType.tp_doc = "Type for all Unreal exposed map instances";
 
 	static PySequenceMethods PySequence;
 	PySequence.sq_contains = (objobjproc)&FSequenceFuncs::Contains;
@@ -1683,7 +1683,7 @@ PyTypeObject InitializePyWrapperMapIteratorType(const char* InName)
 	PyType.tp_iternext = (iternextfunc)&FFuncs::IterNext;
 
 	PyType.tp_flags = Py_TPFLAGS_DEFAULT;
-	PyType.tp_doc = "Type for all UE4 exposed map iterators";
+	PyType.tp_doc = "Type for all Unreal exposed map iterators";
 
 	return PyType;
 }
@@ -1757,7 +1757,7 @@ PyTypeObject InitializePyWrapperMapViewType(const char* InName)
 	PyType.tp_iter = (getiterfunc)&FFuncs::GetIter;
 
 	PyType.tp_flags = Py_TPFLAGS_DEFAULT;
-	PyType.tp_doc = "Type for all UE4 exposed map views";
+	PyType.tp_doc = "Type for all Unreal exposed map views";
 
 	static PySequenceMethods PySequence;
 	PySequence.sq_length = (lenfunc)&FSequenceFuncs::Len;
diff --git a/Engine/Plugins/Experimental/PythonScriptPlugin/Source/PythonScriptPlugin/Private/PyWrapperMap.h b/Engine/Plugins/Experimental/PythonScriptPlugin/Source/PythonScriptPlugin/Private/PyWrapperMap.h
--- a/Engine/Plugins/Experimental/PythonScriptPlugin/Source/PythonScriptPlugin/Private/PyWrapperMap.h
+++ b/Engine/Plugins/Experimental/PythonScriptPlugin/Source/PythonScriptPlugin/Private/PyWrapperMap.h
@@ -14,7 +14,7 @@ extern PyTypeObject PyWrapperMapType;
 /** Initialize the PyWrapperMap types and add them to the given Python module */
 void InitializePyWrapperMap(PyGenUtil::FNativePythonModule& ModuleInfo);
 
-/** Type for all UE4 exposed map instances */
+/** Type for all Unreal exposed map instances */
 struct FPyWrapperMap : public FPyWrapperBase
 {
 	/** The owner of the wrapped map instance (if any) */
@@ -93,7 +93,7 @@ struct FPyWrapperMap : public FPyWrapperBase
 	static FPyWrapperMap* FromKeys(PyObject* InSequence, PyObject* InValue, PyTypeObject* InType);
 };
 
-/** Meta-data for all UE4 exposed map types */
+/** Meta-data for all Unreal exposed map types */
 struct FPyWrapperMapMetaData : public FPyWrapperBaseMetaData
 {
 	PY_METADATA_METHODS(FPyWrapperMapMetaData, FGuid(0xA6AAB38C, 0x5C174F81, 0xB80FA903, 0xBF50C8A2))
diff --git a/Engine/Plugins/Experimental/PythonScriptPlugin/Source/PythonScriptPlugin/Private/PyWrapperName.cpp b/Engine/Plugins/Experimental/PythonScriptPlugin/Source/PythonScriptPlugin/Private/PyWrapperName.cpp
--- a/Engine/Plugins/Experimental/PythonScriptPlugin/Source/PythonScriptPlugin/Private/PyWrapperName.cpp
+++ b/Engine/Plugins/Experimental/PythonScriptPlugin/Source/PythonScriptPlugin/Private/PyWrapperName.cpp
@@ -167,7 +167,7 @@ PyTypeObject InitializePyWrapperNameType()
 		{ nullptr, nullptr, 0, nullptr }
 	};
 
-	PyTypeObject PyType = InitializePyWrapperBasicType<FPyWrapperName>("Name", "Type for all UE4 exposed name instances");
+	PyTypeObject PyType = InitializePyWrapperBasicType<FPyWrapperName>("Name", "Type for all Unreal exposed name instances");
 
 	PyType.tp_init = (initproc)&FFuncs::Init;
 	PyType.tp_str = (reprfunc)&FFuncs::Str;
diff --git a/Engine/Plugins/Experimental/PythonScriptPlugin/Source/PythonScriptPlugin/Private/PyWrapperName.h b/Engine/Plugins/Experimental/PythonScriptPlugin/Source/PythonScriptPlugin/Private/PyWrapperName.h
--- a/Engine/Plugins/Experimental/PythonScriptPlugin/Source/PythonScriptPlugin/Private/PyWrapperName.h
+++ b/Engine/Plugins/Experimental/PythonScriptPlugin/Source/PythonScriptPlugin/Private/PyWrapperName.h
@@ -12,7 +12,7 @@ extern PyTypeObject PyWrapperNameType;
 /** Initialize the FPyWrapperName types and add them to the given Python module */
 void InitializePyWrapperName(PyGenUtil::FNativePythonModule& ModuleInfo);
 
-/** Type for all UE4 exposed FName instances */
+/** Type for all Unreal exposed FName instances */
 struct FPyWrapperName : public TPyWrapperBasic<FName, FPyWrapperName>
 {
 	typedef TPyWrapperBasic<FName, FPyWrapperName> Super;
diff --git a/Engine/Plugins/Experimental/PythonScriptPlugin/Source/PythonScriptPlugin/Private/PyWrapperObject.cpp b/Engine/Plugins/Experimental/PythonScriptPlugin/Source/PythonScriptPlugin/Private/PyWrapperObject.cpp
--- a/Engine/Plugins/Experimental/PythonScriptPlugin/Source/PythonScriptPlugin/Private/PyWrapperObject.cpp
+++ b/Engine/Plugins/Experimental/PythonScriptPlugin/Source/PythonScriptPlugin/Private/PyWrapperObject.cpp
@@ -1010,7 +1010,7 @@ PyTypeObject InitializePyWrapperObjectType()
 	PyType.tp_methods = PyMethods;
 
 	PyType.tp_flags = Py_TPFLAGS_DEFAULT | Py_TPFLAGS_BASETYPE;
-	PyType.tp_doc = "Type for all UE4 exposed object instances";
+	PyType.tp_doc = "Type for all Unreal exposed object instances";
 
 	return PyType;
 }
diff --git a/Engine/Plugins/Experimental/PythonScriptPlugin/Source/PythonScriptPlugin/Private/PyWrapperObject.h b/Engine/Plugins/Experimental/PythonScriptPlugin/Source/PythonScriptPlugin/Private/PyWrapperObject.h
--- a/Engine/Plugins/Experimental/PythonScriptPlugin/Source/PythonScriptPlugin/Private/PyWrapperObject.h
+++ b/Engine/Plugins/Experimental/PythonScriptPlugin/Source/PythonScriptPlugin/Private/PyWrapperObject.h
@@ -16,7 +16,7 @@ extern PyTypeObject PyWrapperObjectType;
 /** Initialize the PyWrapperObject types and add them to the given Python module */
 void InitializePyWrapperObject(PyGenUtil::FNativePythonModule& ModuleInfo);
 
-/** Type for all UE4 exposed object instances */
+/** Type for all Unreal exposed object instances */
 struct FPyWrapperObject : public FPyWrapperBase
 {
 	/** Wrapped object instance */
@@ -101,7 +101,7 @@ struct FPyWrapperObject : public FPyWrapperBase
 	static int Setter_Impl(FPyWrapperObject* InSelf, PyObject* InValue, void* InClosure);
 };
 
-/** Meta-data for all UE4 exposed object types */
+/** Meta-data for all Unreal exposed object types */
 struct FPyWrapperObjectMetaData : public FPyWrapperBaseMetaData
 {
 	PY_METADATA_METHODS(FPyWrapperObjectMetaData, FGuid(0x89FC2465, 0xA83F4F31, 0xBBCC1E86, 0xE9D76551))
diff --git a/Engine/Plugins/Experimental/PythonScriptPlugin/Source/PythonScriptPlugin/Private/PyWrapperSet.cpp b/Engine/Plugins/Experimental/PythonScriptPlugin/Source/PythonScriptPlugin/Private/PyWrapperSet.cpp
--- a/Engine/Plugins/Experimental/PythonScriptPlugin/Source/PythonScriptPlugin/Private/PyWrapperSet.cpp
+++ b/Engine/Plugins/Experimental/PythonScriptPlugin/Source/PythonScriptPlugin/Private/PyWrapperSet.cpp
@@ -1518,7 +1518,7 @@ PyTypeObject InitializePyWrapperSetType()
 #if PY_MAJOR_VERSION < 3
 	PyType.tp_flags |= Py_TPFLAGS_CHECKTYPES;
 #endif	// PY_MAJOR_VERSION < 3
-	PyType.tp_doc = "Type for all UE4 exposed set instances";
+	PyType.tp_doc = "Type for all Unreal exposed set instances";
 
 	static PyNumberMethods PyNumber;
 	PyNumber.nb_subtract = (binaryfunc)&FNumberFuncs::Sub;
@@ -1595,7 +1595,7 @@ PyTypeObject InitializePyWrapperSetIteratorType()
 	PyType.tp_iternext = (iternextfunc)&FFuncs::IterNext;
 
 	PyType.tp_flags = Py_TPFLAGS_DEFAULT;
-	PyType.tp_doc = "Type for all UE4 exposed set iterators";
+	PyType.tp_doc = "Type for all Unreal exposed set iterators";
 
 	return PyType;
 }
diff --git a/Engine/Plugins/Experimental/PythonScriptPlugin/Source/PythonScriptPlugin/Private/PyWrapperSet.h b/Engine/Plugins/Experimental/PythonScriptPlugin/Source/PythonScriptPlugin/Private/PyWrapperSet.h
--- a/Engine/Plugins/Experimental/PythonScriptPlugin/Source/PythonScriptPlugin/Private/PyWrapperSet.h
+++ b/Engine/Plugins/Experimental/PythonScriptPlugin/Source/PythonScriptPlugin/Private/PyWrapperSet.h
@@ -14,7 +14,7 @@ extern PyTypeObject PyWrapperSetType;
 /** Initialize the PyWrapperSet types and add them to the given Python module */
 void InitializePyWrapperSet(PyGenUtil::FNativePythonModule& ModuleInfo);
 
-/** Type for all UE4 exposed set instances */
+/** Type for all Unreal exposed set instances */
 struct FPyWrapperSet : public FPyWrapperBase
 {
 	/** The owner of the wrapped set instance (if any) */
@@ -105,7 +105,7 @@ struct FPyWrapperSet : public FPyWrapperBase
 	static int IsSuperset(FPyWrapperSet* InSelf, PyObject* InOther);
 };
 
-/** Meta-data for all UE4 exposed set types */
+/** Meta-data for all Unreal exposed set types */
 struct FPyWrapperSetMetaData : public FPyWrapperBaseMetaData
 {
 	PY_METADATA_METHODS(FPyWrapperSetMetaData, FGuid(0xB31DDC32, 0xAEA548BE, 0xB6FB5AB2, 0x02E5CA83))
diff --git a/Engine/Plugins/Experimental/PythonScriptPlugin/Source/PythonScriptPlugin/Private/PyWrapperStruct.cpp b/Engine/Plugins/Experimental/PythonScriptPlugin/Source/PythonScriptPlugin/Private/PyWrapperStruct.cpp
--- a/Engine/Plugins/Experimental/PythonScriptPlugin/Source/PythonScriptPlugin/Private/PyWrapperStruct.cpp
+++ b/Engine/Plugins/Experimental/PythonScriptPlugin/Source/PythonScriptPlugin/Private/PyWrapperStruct.cpp
@@ -1227,7 +1227,7 @@ PyTypeObject InitializePyWrapperStructType()
 #if PY_MAJOR_VERSION < 3
 	PyType.tp_flags |= Py_TPFLAGS_CHECKTYPES;
 #endif	// PY_MAJOR_VERSION < 3
-	PyType.tp_doc = "Type for all UE4 exposed struct instances";
+	PyType.tp_doc = "Type for all Unreal exposed struct instances";
 
 	static PyNumberMethods PyNumber;
 #if PY_MAJOR_VERSION >= 3
diff --git a/Engine/Plugins/Experimental/PythonScriptPlugin/Source/PythonScriptPlugin/Private/PyWrapperStruct.h b/Engine/Plugins/Experimental/PythonScriptPlugin/Source/PythonScriptPlugin/Private/PyWrapperStruct.h
--- a/Engine/Plugins/Experimental/PythonScriptPlugin/Source/PythonScriptPlugin/Private/PyWrapperStruct.h
+++ b/Engine/Plugins/Experimental/PythonScriptPlugin/Source/PythonScriptPlugin/Private/PyWrapperStruct.h
@@ -19,7 +19,7 @@ extern PyTypeObject PyWrapperStructType;
 /** Initialize the PyWrapperStruct types and add them to the given Python module */
 void InitializePyWrapperStruct(PyGenUtil::FNativePythonModule& ModuleInfo);
 
-/** Type for all UE4 exposed struct instances */
+/** Type for all Unreal exposed struct instances */
 struct FPyWrapperStruct : public FPyWrapperBase
 {
 	/** The owner of the wrapped struct instance (if any) */
@@ -298,7 +298,7 @@ public:
 	}
 };
 
-/** Meta-data for all UE4 exposed struct types */
+/** Meta-data for all Unreal exposed struct types */
 struct FPyWrapperStructMetaData : public FPyWrapperBaseMetaData
 {
 	PY_METADATA_METHODS(FPyWrapperStructMetaData, FGuid(0x03C9EA75, 0x2C86448B, 0xB53D1453, 0x94AA6BAE))
diff --git a/Engine/Plugins/Experimental/PythonScriptPlugin/Source/PythonScriptPlugin/Private/PyWrapperText.cpp b/Engine/Plugins/Experimental/PythonScriptPlugin/Source/PythonScriptPlugin/Private/PyWrapperText.cpp
--- a/Engine/Plugins/Experimental/PythonScriptPlugin/Source/PythonScriptPlugin/Private/PyWrapperText.cpp
+++ b/Engine/Plugins/Experimental/PythonScriptPlugin/Source/PythonScriptPlugin/Private/PyWrapperText.cpp
@@ -465,7 +465,7 @@ PyTypeObject InitializePyWrapperTextType()
 		{ nullptr, nullptr, 0, nullptr }
 	};
 
-	PyTypeObject PyType = InitializePyWrapperBasicType<FPyWrapperText>("Text", "Type for all UE4 exposed text instances");
+	PyTypeObject PyType = InitializePyWrapperBasicType<FPyWrapperText>("Text", "Type for all Unreal exposed text instances");
 
 	PyType.tp_init = (initproc)&FFuncs::Init;
 	PyType.tp_str = (reprfunc)&FFuncs::Str;
diff --git a/Engine/Plugins/Experimental/PythonScriptPlugin/Source/PythonScriptPlugin/Private/PyWrapperText.h b/Engine/Plugins/Experimental/PythonScriptPlugin/Source/PythonScriptPlugin/Private/PyWrapperText.h
--- a/Engine/Plugins/Experimental/PythonScriptPlugin/Source/PythonScriptPlugin/Private/PyWrapperText.h
+++ b/Engine/Plugins/Experimental/PythonScriptPlugin/Source/PythonScriptPlugin/Private/PyWrapperText.h
@@ -12,7 +12,7 @@ extern PyTypeObject PyWrapperTextType;
 /** Initialize the FPyWrapperText types and add them to the given Python module */
 void InitializePyWrapperText(PyGenUtil::FNativePythonModule& ModuleInfo);
 
-/** Type for all UE4 exposed FText instances */
+/** Type for all Unreal exposed FText instances */
 struct FPyWrapperText : public TPyWrapperBasic<FText, FPyWrapperText>
 {
 	typedef TPyWrapperBasic<FText, FPyWrapperText> Super;
diff --git a/Engine/Plugins/Experimental/PythonScriptPlugin/Source/PythonScriptPlugin/Private/PythonScriptPlugin.cpp b/Engine/Plugins/Experimental/PythonScriptPlugin/Source/PythonScriptPlugin/Private/PythonScriptPlugin.cpp
--- a/Engine/Plugins/Experimental/PythonScriptPlugin/Source/PythonScriptPlugin/Private/PythonScriptPlugin.cpp
+++ b/Engine/Plugins/Experimental/PythonScriptPlugin/Source/PythonScriptPlugin/Private/PythonScriptPlugin.cpp
@@ -618,7 +618,7 @@ void FPythonScriptPlugin::InitializePython()
 
 	// Set-up the correct home path
 	{
-		// Build the full Python directory (UE_PYTHON_DIR may be relative to UE4 engine directory for portability)
+		// Build the full Python directory (UE_PYTHON_DIR may be relative to the engine directory for portability)
 		FString PythonDir = UTF8_TO_TCHAR(UE_PYTHON_DIR);
 		PythonDir.ReplaceInline(TEXT("{ENGINE_DIR}"), *FPaths::EngineDir(), ESearchCase::CaseSensitive);
 		FPaths::NormalizeDirectoryName(PythonDir);
@@ -630,7 +630,7 @@ void FPythonScriptPlugin::InitializePython()
 	{
 		UE_LOG(LogPython, Log, TEXT("Using Python %d.%d.%d"), PY_MAJOR_VERSION, PY_MINOR_VERSION, PY_MICRO_VERSION);
 
-		// Python 3 changes the console mode from O_TEXT to O_BINARY which affects other UE4 uses of the console
+		// Python 3 changes the console mode from O_TEXT to O_BINARY which affects other uses of the console
 		// So change the console mode back to its current setting after Py_Initialize has been called
 #if PLATFORM_WINDOWS && PY_MAJOR_VERSION >= 3
 		// We call _setmode here to cache the current state
@@ -646,7 +646,7 @@ void FPythonScriptPlugin::InitializePython()
 #endif	// PLATFORM_WINDOWS && PY_MAJOR_VERSION >= 3
 
 #if PY_MAJOR_VERSION >= 3 && PY_MINOR_VERSION >= 7
-		// Python 3.7+ changes the C locale which affects UE4 functions using C string APIs
+		// Python 3.7+ changes the C locale which affects functions using C string APIs
 		// So change the C locale back to its current setting after Py_Initialize has been called
 		FString CurrentLocale;
 		if (const char* CurrentLocalePtr = setlocale(LC_ALL, nullptr))
@@ -660,9 +660,9 @@ void FPythonScriptPlugin::InitializePython()
 #endif	// PY_MAJOR_VERSION >= 3 && PY_MINOR_VERSION >= 4
 		Py_SetProgramName(PyProgramName.GetData());
 		Py_SetPythonHome(PyHomePath.GetData());
-		Py_InitializeEx(0); // 0 so Python doesn't override any UE4 signal handling
+		Py_InitializeEx(0); // 0 so Python doesn't override any signal handling
 
-		// Ensure Python supports multiple threads via the GIL, as UE4 GC runs over multiple threads, 
+		// Ensure Python supports multiple threads via the GIL, as UE GC runs over multiple threads, 
 		// which may invoke FPyReferenceCollector::AddReferencedObjects on a background thread...
 		if (!PyEval_ThreadsInitialized())
 		{
@@ -1287,7 +1287,7 @@ bool FPythonScriptPlugin::RunFile(const TCHAR* InFile, const TCHAR* InArgs, FPyt
 			FPythonScopedArgv ScopedArgv(InArgs);
 
 			FDelegateHandle LogCaptureHandle = PyCore::GetPythonLogCapture().AddLambda([&InOutPythonCommand](EPythonLogOutputType InLogType, const TCHAR* InLogString) { InOutPythonCommand.LogOutput.Add(FPythonLogOutputEntry{ InLogType, InLogString }); });
-			PyResult = FPyObjectPtr::StealReference(EvalString(*FileStr, *ResolvedFilePath, Py_file_input, PyFileGlobalDict, PyFileLocalDict)); // We can't just use PyRun_File here as Python isn't always built against the same version of the CRT as UE4, so we get a crash at the CRT layer
+			PyResult = FPyObjectPtr::StealReference(EvalString(*FileStr, *ResolvedFilePath, Py_file_input, PyFileGlobalDict, PyFileLocalDict)); // We can't just use PyRun_File here as Python isn't always built against the same version of the CRT as UE, so we get a crash at the CRT layer
 			PyCore::GetPythonLogCapture().Remove(LogCaptureHandle);
 		}
 
diff --git a/Engine/Plugins/Experimental/PythonScriptPlugin/Source/PythonScriptPluginPreload/Private/PythonScriptPluginPreload.cpp b/Engine/Plugins/Experimental/PythonScriptPlugin/Source/PythonScriptPluginPreload/Private/PythonScriptPluginPreload.cpp
--- a/Engine/Plugins/Experimental/PythonScriptPlugin/Source/PythonScriptPluginPreload/Private/PythonScriptPluginPreload.cpp
+++ b/Engine/Plugins/Experimental/PythonScriptPlugin/Source/PythonScriptPluginPreload/Private/PythonScriptPluginPreload.cpp
@@ -75,7 +75,7 @@ private:
 #if PLATFORM_WINDOWS || PLATFORM_LINUX
 		// Load the DSOs
 		{
-			// Build the full Python directory (UE_PYTHON_DIR may be relative to UE4 engine directory for portability)
+			// Build the full Python directory (UE_PYTHON_DIR may be relative to the engine directory for portability)
 			FString PythonDir = UTF8_TO_TCHAR(UE_PYTHON_DIR);
 			PythonDir.ReplaceInline(TEXT("{ENGINE_DIR}"), *FPaths::EngineDir(), ESearchCase::CaseSensitive);
 			FPaths::NormalizeDirectoryName(PythonDir);
