Generating functions from inherited UInterfaces when generating glue for Python scripting

diff --git a/Engine/Plugins/Experimental/PythonScriptPlugin/Source/PythonScriptPlugin/Private/PyWrapperTypeRegistry.cpp b/Engine/Plugins/Experimental/PythonScriptPlugin/Source/PythonScriptPlugin/Private/PyWrapperTypeRegistry.cpp
--- a/Engine/Plugins/Experimental/PythonScriptPlugin/Source/PythonScriptPlugin/Private/PyWrapperTypeRegistry.cpp
+++ b/Engine/Plugins/Experimental/PythonScriptPlugin/Source/PythonScriptPlugin/Private/PyWrapperTypeRegistry.cpp
@@ -1254,7 +1254,7 @@ PyTypeObject* FPyWrapperTypeRegistry::GenerateWrappedClassType(const UClass* InC
 			continue;
 		}
 	}
-	for (TFieldIterator<const UField> FieldIt(InClass, EFieldIteratorFlags::ExcludeSuper); FieldIt; ++FieldIt)
+	for (TFieldIterator<const UField> FieldIt(InClass, EFieldIteratorFlags::ExcludeSuper, EFieldIteratorFlags::IncludeDeprecated, EFieldIteratorFlags::IncludeInterfaces); FieldIt; ++FieldIt)
 	{
 		if (const UFunction* Func = Cast<const UFunction>(*FieldIt))
 		{
