Fixed warnings emitted from exporting interfaces to Python

 - Only export interface functions if there isn't already a function with this name on the type.
 - Added tests for interfaces, including interface inheritance, and multiple inheritance on an object.

diff --git a/Engine/Plugins/Experimental/PythonScriptPlugin/Source/PythonScriptPlugin/Private/PyGenUtil.cpp b/Engine/Plugins/Experimental/PythonScriptPlugin/Source/PythonScriptPlugin/Private/PyGenUtil.cpp
--- a/Engine/Plugins/Experimental/PythonScriptPlugin/Source/PythonScriptPlugin/Private/PyGenUtil.cpp
+++ b/Engine/Plugins/Experimental/PythonScriptPlugin/Source/PythonScriptPlugin/Private/PyGenUtil.cpp
@@ -1475,25 +1475,21 @@ bool IsScriptExposedFunction(const UFunction* InFunc)
 
 bool HasScriptExposedFields(const UStruct* InStruct)
 {
-	for (TFieldIterator<const UField> FieldIt(InStruct); FieldIt; ++FieldIt)
+	for (TFieldIterator<const UFunction> FieldIt(InStruct); FieldIt; ++FieldIt)
 	{
-		if (const UFunction* Func = Cast<const UFunction>(*FieldIt))
+		const UFunction* Func = *FieldIt;
+		if (IsScriptExposedFunction(Func))
 		{
-			if (IsScriptExposedFunction(Func))
-			{
-				return true;
-			}
+			return true;
 		}
 	}
 
-	for (TFieldIterator<const FField> FieldIt(InStruct); FieldIt; ++FieldIt)
+	for (TFieldIterator<const FProperty> FieldIt(InStruct); FieldIt; ++FieldIt)
 	{
-		if (const FProperty* Prop = CastField<const FProperty>(*FieldIt))
+		const FProperty* Prop = *FieldIt;
+		if (IsScriptExposedProperty(Prop))
 		{
-			if (IsScriptExposedProperty(Prop))
-			{
-				return true;
-			}
+			return true;
 		}
 	}
 
@@ -1574,6 +1570,51 @@ bool IsDeprecatedFunction(const UFunction* InFunc, FString* OutDeprecationMessag
 	return false;
 }
 
+void GetExportedInterfacesForClass_Recursive(const UClass* InInterface, TSet<const UClass*>& InOutProcessedInterfaces, TArray<const UClass*>& InOutExportedInterfaces)
+{
+	check(InInterface->HasAnyClassFlags(CLASS_Interface));
+
+	{
+		bool bAlreadyProcessed = false;
+		InOutProcessedInterfaces.Add(InInterface, &bAlreadyProcessed);
+		if (bAlreadyProcessed)
+		{
+			return;
+		}
+	}
+
+	if (!ShouldExportClass(InInterface))
+	{
+		return;
+	}
+
+	InOutExportedInterfaces.Add(InInterface);
+
+	for (UClass* SuperClass = InInterface->GetSuperClass(); 
+		SuperClass && SuperClass->HasAnyClassFlags(CLASS_Interface); 
+		SuperClass = SuperClass->GetSuperClass()
+		)
+	{
+		GetExportedInterfacesForClass_Recursive(SuperClass, InOutProcessedInterfaces, InOutExportedInterfaces);
+	}
+
+	for (const FImplementedInterface& Interface : InInterface->Interfaces)
+	{
+		GetExportedInterfacesForClass_Recursive(Interface.Class, InOutProcessedInterfaces, InOutExportedInterfaces);
+	}
+}
+
+TArray<const UClass*> GetExportedInterfacesForClass(const UClass* InClass)
+{
+	TSet<const UClass*> ProcessedInterfaces;
+	TArray<const UClass*> ExportedInterfaces;
+	for (const FImplementedInterface& Interface : InClass->Interfaces)
+	{
+		GetExportedInterfacesForClass_Recursive(Interface.Class, ProcessedInterfaces, ExportedInterfaces);
+	}
+	return ExportedInterfaces;
+}
+
 bool ShouldExportClass(const UClass* InClass)
 {
 	return IsScriptExposedClass(InClass) || HasScriptExposedFields(InClass);
diff --git a/Engine/Plugins/Experimental/PythonScriptPlugin/Source/PythonScriptPlugin/Private/PyGenUtil.h b/Engine/Plugins/Experimental/PythonScriptPlugin/Source/PythonScriptPlugin/Private/PyGenUtil.h
--- a/Engine/Plugins/Experimental/PythonScriptPlugin/Source/PythonScriptPlugin/Private/PyGenUtil.h
+++ b/Engine/Plugins/Experimental/PythonScriptPlugin/Source/PythonScriptPlugin/Private/PyGenUtil.h
@@ -1037,6 +1037,9 @@ namespace PyGenUtil
 	/** Is the given function marked as deprecated? */
 	bool IsDeprecatedFunction(const UFunction* InFunc, FString* OutDeprecationMessage = nullptr);
 
+	/** Given a class, get all the interface classes related to it that should be included when exporting it to Python */
+	TArray<const UClass*> GetExportedInterfacesForClass(const UClass* InClass);
+
 	/** Should the given class be exported to Python? */
 	bool ShouldExportClass(const UClass* InClass);
 
diff --git a/Engine/Plugins/Experimental/PythonScriptPlugin/Source/PythonScriptPlugin/Private/PyTest.cpp b/Engine/Plugins/Experimental/PythonScriptPlugin/Source/PythonScriptPlugin/Private/PyTest.cpp
--- a/Engine/Plugins/Experimental/PythonScriptPlugin/Source/PythonScriptPlugin/Private/PyTest.cpp
+++ b/Engine/Plugins/Experimental/PythonScriptPlugin/Source/PythonScriptPlugin/Private/PyTest.cpp
@@ -50,6 +50,21 @@ FPyTestStruct UPyTestStructLibrary::AddStr(const FPyTestStruct& InStruct, const
 	return Result;
 }
 
+UPyTestInterface::UPyTestInterface(const FObjectInitializer& ObjectInitializer)
+	: Super(ObjectInitializer)
+{
+}
+
+UPyTestChildInterface::UPyTestChildInterface(const FObjectInitializer& ObjectInitializer)
+	: Super(ObjectInitializer)
+{
+}
+
+UPyTestOtherInterface::UPyTestOtherInterface(const FObjectInitializer& ObjectInitializer)
+	: Super(ObjectInitializer)
+{
+}
+
 UPyTestObject::UPyTestObject()
 {
 	StructArray.AddDefaulted();
@@ -152,6 +167,21 @@ int32 UPyTestObject::GetConstantValue()
 	return 10;
 }
 
+int32 UPyTestObject::FuncInterface(const int32 InValue) const
+{
+	return InValue;
+}
+
+int32 UPyTestObject::FuncInterfaceChild(const int32 InValue) const
+{
+	return InValue;
+}
+
+int32 UPyTestObject::FuncInterfaceOther(const int32 InValue) const
+{
+	return InValue;
+}
+
 bool UPyTestObjectLibrary::IsBoolSet(const UPyTestObject* InObj)
 {
 	return InObj->Bool;
diff --git a/Engine/Plugins/Experimental/PythonScriptPlugin/Source/PythonScriptPlugin/Private/PyTest.h b/Engine/Plugins/Experimental/PythonScriptPlugin/Source/PythonScriptPlugin/Private/PyTest.h
--- a/Engine/Plugins/Experimental/PythonScriptPlugin/Source/PythonScriptPlugin/Private/PyTest.h
+++ b/Engine/Plugins/Experimental/PythonScriptPlugin/Source/PythonScriptPlugin/Private/PyTest.h
@@ -3,6 +3,7 @@
 #pragma once
 
 #include "CoreMinimal.h"
+#include "PyTestInterface.h"
 #include "Kismet/BlueprintFunctionLibrary.h"
 #include "PyTest.generated.h"
 
@@ -117,7 +118,7 @@ class UPyTestStructLibrary : public UBlueprintFunctionLibrary
  * Object to allow testing of the various UObject features that are exposed to Python wrapped types.
  */
 UCLASS(Blueprintable)
-class UPyTestObject : public UObject
+class UPyTestObject : public UObject, public IPyTestChildInterface, public IPyTestOtherInterface
 {
 	GENERATED_BODY()
 
@@ -228,6 +229,10 @@ public:
 
 	UFUNCTION(BlueprintPure, Category = "Python|Internal", meta=(ScriptConstant="ConstantValue"))
 	static int32 GetConstantValue();
+
+	virtual int32 FuncInterface(const int32 InValue) const override;
+	virtual int32 FuncInterfaceChild(const int32 InValue) const override;
+	virtual int32 FuncInterfaceOther(const int32 InValue) const override;
 };
 
 /**
diff --git a/Engine/Plugins/Experimental/PythonScriptPlugin/Source/PythonScriptPlugin/Private/PyTestInterface.h b/Engine/Plugins/Experimental/PythonScriptPlugin/Source/PythonScriptPlugin/Private/PyTestInterface.h
new file mode 100644
--- /dev/null
+++ b/Engine/Plugins/Experimental/PythonScriptPlugin/Source/PythonScriptPlugin/Private/PyTestInterface.h
@@ -0,0 +1,58 @@
+// Copyright Epic Games, Inc. All Rights Reserved.
+
+#pragma once
+
+#include "CoreMinimal.h"
+#include "UObject/Interface.h"
+#include "PyTestInterface.generated.h"
+
+/**
+ * Interface to allow testing of the various UInterface features that are exposed to Python wrapped types.
+ */
+UINTERFACE(BlueprintType, meta=(CannotImplementInterfaceInBlueprint))
+class UPyTestInterface : public UInterface
+{
+	GENERATED_UINTERFACE_BODY()
+};
+class IPyTestInterface
+{
+	GENERATED_IINTERFACE_BODY()
+
+public:
+	UFUNCTION(BlueprintCallable, Category = "Python|Internal")
+	virtual int32 FuncInterface(const int32 InValue) const = 0;
+};
+
+/**
+ * Interface to allow testing of inheritance on Python wrapped types.
+ */
+UINTERFACE(BlueprintType, meta=(CannotImplementInterfaceInBlueprint))
+class UPyTestChildInterface : public UPyTestInterface
+{
+	GENERATED_UINTERFACE_BODY()
+};
+class IPyTestChildInterface : public IPyTestInterface
+{
+	GENERATED_IINTERFACE_BODY()
+
+public:
+	UFUNCTION(BlueprintCallable, Category = "Python|Internal")
+	virtual int32 FuncInterfaceChild(const int32 InValue) const = 0;
+};
+
+/**
+ * Interface to allow testing of multiple-inheritance on Python wrapped types.
+ */
+UINTERFACE(BlueprintType, meta=(CannotImplementInterfaceInBlueprint))
+class UPyTestOtherInterface : public UInterface
+{
+	GENERATED_UINTERFACE_BODY()
+};
+class IPyTestOtherInterface
+{
+	GENERATED_IINTERFACE_BODY()
+
+public:
+	UFUNCTION(BlueprintCallable, Category = "Python|Internal")
+	virtual int32 FuncInterfaceOther(const int32 InValue) const = 0;
+};
diff --git a/Engine/Plugins/Experimental/PythonScriptPlugin/Source/PythonScriptPlugin/Private/PyWrapperTypeRegistry.cpp b/Engine/Plugins/Experimental/PythonScriptPlugin/Source/PythonScriptPlugin/Private/PyWrapperTypeRegistry.cpp
--- a/Engine/Plugins/Experimental/PythonScriptPlugin/Source/PythonScriptPlugin/Private/PyWrapperTypeRegistry.cpp
+++ b/Engine/Plugins/Experimental/PythonScriptPlugin/Source/PythonScriptPlugin/Private/PyWrapperTypeRegistry.cpp
@@ -1246,20 +1246,32 @@ PyTypeObject* FPyWrapperTypeRegistry::GenerateWrappedClassType(const UClass* InC
 	const FString PythonClassName = PyGenUtil::GetClassPythonName(InClass);
 	GeneratedWrappedType->TypeName = PyGenUtil::TCHARToUTF8Buffer(*PythonClassName);
 
-	for (TFieldIterator<const FField> FieldIt(InClass, EFieldIteratorFlags::ExcludeSuper); FieldIt; ++FieldIt)
+	for (TFieldIterator<const FProperty> FieldIt(InClass, EFieldIteratorFlags::ExcludeSuper); FieldIt; ++FieldIt)
 	{
-		if (const FProperty* Prop = CastField<const FProperty>(*FieldIt))
-		{
-			GenerateWrappedProperty(Prop);
-			continue;
-		}
+		const FProperty* Prop = *FieldIt;
+		GenerateWrappedProperty(Prop);
+	}
+
+	for (TFieldIterator<const UFunction> FieldIt(InClass, EFieldIteratorFlags::ExcludeSuper, EFieldIteratorFlags::IncludeDeprecated); FieldIt; ++FieldIt)
+	{
+		const UFunction* Func = *FieldIt;
+		GenerateWrappedMethod(Func);
 	}
-	for (TFieldIterator<const UField> FieldIt(InClass, EFieldIteratorFlags::ExcludeSuper, EFieldIteratorFlags::IncludeDeprecated, EFieldIteratorFlags::IncludeInterfaces); FieldIt; ++FieldIt)
+
 	{
-		if (const UFunction* Func = Cast<const UFunction>(*FieldIt))
+		TArray<const UClass*> Interfaces = PyGenUtil::GetExportedInterfacesForClass(InClass);
+		for (const UClass* Interface : Interfaces)
 		{
-			GenerateWrappedMethod(Func);
-			continue;
+			for (TFieldIterator<const UFunction> FieldIt(Interface, EFieldIteratorFlags::ExcludeSuper, EFieldIteratorFlags::IncludeDeprecated); FieldIt; ++FieldIt)
+			{
+				const UFunction* Func = *FieldIt;
+				
+				// Skip the interface function if the current class already has a reflected function with this name
+				if (!InClass->FindFunctionByName(Func->GetFName(), EIncludeSuperFlag::ExcludeSuper))
+				{
+					GenerateWrappedMethod(Func);
+				}
+			}
 		}
 	}
 
