Isolate the embedded Python interpreter from the general Python system environment

The following flags are now set when initializing the embedded Python interpreter:
 - Py_IgnoreEnvironmentFlag - Ignore PYTHONPATH when calculating sys.path.
 - Py_IsolatedFlag - Skip adding the user's site-packages directory to sys.path.

To counter the loss of PYTHONPATH, a new UE_PYTHONPATH environment variable is now available with the same rules as PYTHONPATH.

diff --git a/Engine/Plugins/Experimental/PythonScriptPlugin/Source/PythonScriptPlugin/Private/PythonScriptPlugin.cpp b/Engine/Plugins/Experimental/PythonScriptPlugin/Source/PythonScriptPlugin/Private/PythonScriptPlugin.cpp
--- a/Engine/Plugins/Experimental/PythonScriptPlugin/Source/PythonScriptPlugin/Private/PythonScriptPlugin.cpp
+++ b/Engine/Plugins/Experimental/PythonScriptPlugin/Source/PythonScriptPlugin/Private/PythonScriptPlugin.cpp
@@ -24,6 +24,7 @@
 #include "Misc/Paths.h"
 #include "Misc/FileHelper.h"
 #include "Misc/PackageName.h"
+#include "HAL/PlatformMisc.h"
 #include "HAL/PlatformProcess.h"
 #include "Containers/Ticker.h"
 #include "Features/IModularFeatures.h"
@@ -655,7 +656,9 @@ void FPythonScriptPlugin::InitializePython()
 		}
 #endif	// PY_MAJOR_VERSION >= 3 && PY_MINOR_VERSION >= 7
 
+		Py_IgnoreEnvironmentFlag = 1; // 1 so Python doesn't inherit its environment variables into our embedded interpreter
 #if PY_MAJOR_VERSION >= 3 && PY_MINOR_VERSION >= 4
+		Py_IsolatedFlag = 1; // 1 so Python doesn't add the user scripts and site-packages paths into our embedded interpreter
 		Py_SetStandardStreamEncoding("utf-8", nullptr);
 #endif	// PY_MAJOR_VERSION >= 3 && PY_MINOR_VERSION >= 4
 		Py_SetProgramName(PyProgramName.GetData());
@@ -747,6 +750,13 @@ void FPythonScriptPlugin::InitializePython()
 			PyUtil::AddSystemPath(FPaths::ConvertRelativePathToFull(AdditionalPath.Path));
 		}
 
+		TArray<FString> SystemEnvPaths;
+		FPlatformMisc::GetEnvironmentVariable(TEXT("UE_PYTHONPATH")).ParseIntoArray(SystemEnvPaths, FPlatformMisc::GetPathVarDelimiter());
+		for (const FString& SystemEnvPath : SystemEnvPaths)
+		{
+			PyUtil::AddSystemPath(SystemEnvPath);
+		}
+
 		FPackageName::OnContentPathMounted().AddRaw(this, &FPythonScriptPlugin::OnContentPathMounted);
 		FPackageName::OnContentPathDismounted().AddRaw(this, &FPythonScriptPlugin::OnContentPathDismounted);
 		FCoreUObjectDelegates::OnPackageReloaded.AddRaw(this, &FPythonScriptPlugin::OnAssetReload);
