Added support for per-plugin Python site-packages

Plugins can now add third-party Python scripts within their Content/Python/Lib/site-packages and Content/Python/Lib/{Platform}/site-packages folders, and these will be automatically available for import within UE.

diff --git a/Engine/Plugins/Experimental/PythonScriptPlugin/Source/PythonScriptPlugin/Private/PyUtil.cpp b/Engine/Plugins/Experimental/PythonScriptPlugin/Source/PythonScriptPlugin/Private/PyUtil.cpp
--- a/Engine/Plugins/Experimental/PythonScriptPlugin/Source/PythonScriptPlugin/Private/PyUtil.cpp
+++ b/Engine/Plugins/Experimental/PythonScriptPlugin/Source/PythonScriptPlugin/Private/PyUtil.cpp
@@ -794,7 +794,7 @@ UObject* GetOwnerObject(PyObject* InPyObj)
 	}
 
 	return nullptr;
-			}
+}
 
 PyObject* GetPropertyValue(const UStruct* InStruct, const void* InStructData, const FProperty* InProp, const char *InAttributeName, PyObject* InOwnerPyObject, const TCHAR* InErrorCtxt)
 {
@@ -1084,6 +1084,22 @@ FString GetInterpreterExecutablePath(bool* OutIsEnginePython)
 	return PythonPath;
 }
 
+void AddSitePackagesPath(const FString& InPath)
+{
+	if (FPyObjectPtr PySiteModule = FPyObjectPtr::StealReference(PyImport_ImportModule("site")))
+	{
+		PyObject* PySiteDict = PyModule_GetDict(PySiteModule);
+		if (PyObject* PyAddSiteDirFunc = PyDict_GetItemString(PySiteDict, "addsitedir"))
+		{
+			FPyObjectPtr PyPath;
+			if (PyConversion::Pythonize(InPath, PyPath.Get(), PyConversion::ESetErrorState::No))
+			{
+				FPyObjectPtr PyAddSiteDirResult = FPyObjectPtr::StealReference(PyObject_CallFunctionObjArgs(PyAddSiteDirFunc, PyPath.Get(), nullptr));
+			}
+		}
+	}
+}
+
 void AddSystemPath(const FString& InPath)
 {
 	if (PyObject* PyPathList = PySys_GetObject(PyCStrCast("path")))
diff --git a/Engine/Plugins/Experimental/PythonScriptPlugin/Source/PythonScriptPlugin/Private/PyUtil.h b/Engine/Plugins/Experimental/PythonScriptPlugin/Source/PythonScriptPlugin/Private/PyUtil.h
--- a/Engine/Plugins/Experimental/PythonScriptPlugin/Source/PythonScriptPlugin/Private/PyUtil.h
+++ b/Engine/Plugins/Experimental/PythonScriptPlugin/Source/PythonScriptPlugin/Private/PyUtil.h
@@ -467,6 +467,11 @@ namespace PyUtil
 	 */
 	FString GetInterpreterExecutablePath(bool* OutIsEnginePython = nullptr);
 
+	/**
+	 * Register the given path as a known site-packages path (@see site.addsitedir in Python).
+	 */
+	void AddSitePackagesPath(const FString& InPath);
+
 	/**
 	 * Ensure that the given path is on the sys.path list.
 	 */
diff --git a/Engine/Plugins/Experimental/PythonScriptPlugin/Source/PythonScriptPlugin/Private/PythonScriptPlugin.cpp b/Engine/Plugins/Experimental/PythonScriptPlugin/Source/PythonScriptPlugin/Private/PythonScriptPlugin.cpp
--- a/Engine/Plugins/Experimental/PythonScriptPlugin/Source/PythonScriptPlugin/Private/PythonScriptPlugin.cpp
+++ b/Engine/Plugins/Experimental/PythonScriptPlugin/Source/PythonScriptPlugin/Private/PythonScriptPlugin.cpp
@@ -759,9 +759,7 @@ void FPythonScriptPlugin::InitializePython()
 		for (const FString& RootPath : RootPaths)
 		{
 			const FString RootFilesystemPath = FPackageName::LongPackageNameToFilename(RootPath);
-			const FString PythonContentPath = FPaths::ConvertRelativePathToFull(RootFilesystemPath / TEXT("Python"));
-			PyUtil::AddSystemPath(PythonContentPath);
-			PyUtil::GetOnDiskUnrealModulesCache().AddModules(*PythonContentPath);
+			RegisterModulePaths(RootFilesystemPath);
 		}
 
 		for (const FDirectoryPath& AdditionalPath : GetDefault<UPythonScriptPluginSettings>()->AdditionalPaths)
@@ -1377,9 +1375,7 @@ void FPythonScriptPlugin::OnContentPathMounted(const FString& InAssetPath, const
 {
 	{
 		FPyScopedGIL GIL;
-		const FString PythonContentPath = FPaths::ConvertRelativePathToFull(InFilesystemPath / TEXT("Python"));
-		PyUtil::AddSystemPath(PythonContentPath);
-		PyUtil::GetOnDiskUnrealModulesCache().AddModules(*PythonContentPath);
+		RegisterModulePaths(InFilesystemPath);
 	}
 
 #if WITH_EDITOR
@@ -1394,9 +1390,7 @@ void FPythonScriptPlugin::OnContentPathDismounted(const FString& InAssetPath, co
 {
 	{
 		FPyScopedGIL GIL;
-		const FString PythonContentPath = FPaths::ConvertRelativePathToFull(InFilesystemPath / TEXT("Python"));
-		PyUtil::RemoveSystemPath(PythonContentPath);
-		PyUtil::GetOnDiskUnrealModulesCache().RemoveModules(*PythonContentPath);
+		UnregisterModulePaths(InFilesystemPath);
 	}
 
 #if WITH_EDITOR
@@ -1407,6 +1401,32 @@ void FPythonScriptPlugin::OnContentPathDismounted(const FString& InAssetPath, co
 #endif	// WITH_EDITOR
 }
 
+void FPythonScriptPlugin::RegisterModulePaths(const FString& InFilesystemPath)
+{
+	const FString PythonContentPath = FPaths::ConvertRelativePathToFull(InFilesystemPath / TEXT("Python"));
+	PyUtil::AddSystemPath(PythonContentPath);
+
+	const FString PythonContentPlatformSitePackagesPath = PythonContentPath / TEXT("Lib") / FPlatformMisc::GetUBTPlatform() / TEXT("site-packages");
+	const FString PythonContentGeneralSitePackagesPath = PythonContentPath / TEXT("Lib") / TEXT("site-packages");
+	PyUtil::AddSitePackagesPath(PythonContentPlatformSitePackagesPath);
+	PyUtil::AddSitePackagesPath(PythonContentGeneralSitePackagesPath);
+
+	PyUtil::GetOnDiskUnrealModulesCache().AddModules(*PythonContentPath);
+}
+
+void FPythonScriptPlugin::UnregisterModulePaths(const FString& InFilesystemPath)
+{
+	const FString PythonContentPath = FPaths::ConvertRelativePathToFull(InFilesystemPath / TEXT("Python"));
+	PyUtil::RemoveSystemPath(PythonContentPath);
+
+	const FString PythonContentPlatformSitePackagesPath = PythonContentPath / TEXT("Lib") / FPlatformMisc::GetUBTPlatform() / TEXT("site-packages");
+	const FString PythonContentGeneralSitePackagesPath = PythonContentPath / TEXT("Lib") / TEXT("site-packages");
+	PyUtil::RemoveSystemPath(PythonContentPlatformSitePackagesPath);
+	PyUtil::RemoveSystemPath(PythonContentGeneralSitePackagesPath);
+
+	PyUtil::GetOnDiskUnrealModulesCache().RemoveModules(*PythonContentPath);
+}
+
 bool FPythonScriptPlugin::IsDeveloperModeEnabled()
 {
 	return GetDefault<UPythonScriptPluginSettings>()->bDeveloperMode || GetDefault<UPythonScriptPluginUserSettings>()->bDeveloperMode;
diff --git a/Engine/Plugins/Experimental/PythonScriptPlugin/Source/PythonScriptPlugin/Private/PythonScriptPlugin.h b/Engine/Plugins/Experimental/PythonScriptPlugin/Source/PythonScriptPlugin/Private/PythonScriptPlugin.h
--- a/Engine/Plugins/Experimental/PythonScriptPlugin/Source/PythonScriptPlugin/Private/PythonScriptPlugin.h
+++ b/Engine/Plugins/Experimental/PythonScriptPlugin/Source/PythonScriptPlugin/Private/PythonScriptPlugin.h
@@ -169,6 +169,10 @@ private:
 
 	void OnContentPathDismounted(const FString& InAssetPath, const FString& InFilesystemPath);
 
+	static void RegisterModulePaths(const FString& InFilesystemPath);
+
+	static void UnregisterModulePaths(const FString& InFilesystemPath);
+
 	static bool IsDeveloperModeEnabled();
 
 	void OnAssetRenamed(const FAssetData& Data, const FString& OldName);
