Added missing __repr__ function to the basic Python wrapped types

diff --git a/Engine/Plugins/Experimental/PythonScriptPlugin/Source/PythonScriptPlugin/Private/PyConstant.cpp b/Engine/Plugins/Experimental/PythonScriptPlugin/Source/PythonScriptPlugin/Private/PyConstant.cpp
--- a/Engine/Plugins/Experimental/PythonScriptPlugin/Source/PythonScriptPlugin/Private/PyConstant.cpp
+++ b/Engine/Plugins/Experimental/PythonScriptPlugin/Source/PythonScriptPlugin/Private/PyConstant.cpp
@@ -155,6 +155,7 @@ PyTypeObject InitializePyConstantDescrType()
 
 	PyType.tp_dealloc = (destructor)&FFuncs::Dealloc;
 	PyType.tp_str = (reprfunc)&FFuncs::Str;
+	PyType.tp_repr = (reprfunc)&FFuncs::Str;
 	PyType.tp_descr_get = (descrgetfunc)&FFuncs::DescrGet;
 	PyType.tp_descr_set = (descrsetfunc)&FFuncs::DescrSet;
 	PyType.tp_getattro = (getattrofunc)&PyObject_GenericGetAttr;
diff --git a/Engine/Plugins/Experimental/PythonScriptPlugin/Source/PythonScriptPlugin/Private/PyMethodWithClosure.cpp b/Engine/Plugins/Experimental/PythonScriptPlugin/Source/PythonScriptPlugin/Private/PyMethodWithClosure.cpp
--- a/Engine/Plugins/Experimental/PythonScriptPlugin/Source/PythonScriptPlugin/Private/PyMethodWithClosure.cpp
+++ b/Engine/Plugins/Experimental/PythonScriptPlugin/Source/PythonScriptPlugin/Private/PyMethodWithClosure.cpp
@@ -484,6 +484,7 @@ PyTypeObject InitializePyCFunctionWithClosureType()
 	PyType.tp_traverse = (traverseproc)&FFuncs::Traverse;
 	PyType.tp_clear = (inquiry)&FFuncs::Clear;
 	PyType.tp_str = (reprfunc)&FFuncs::Str;
+	PyType.tp_repr = (reprfunc)&FFuncs::Str;
 	PyType.tp_richcompare = (richcmpfunc)&FFuncs::RichCmp;
 	PyType.tp_hash = (hashfunc)&FFuncs::Hash;
 	PyType.tp_call = (ternaryfunc)&FFuncs::Call;
@@ -606,6 +607,7 @@ PyTypeObject InitializePyMethodWithClosureDescrType()
 	PyType.tp_traverse = (traverseproc)&FFuncs::Traverse;
 	PyType.tp_clear = (inquiry)&FFuncs::Clear;
 	PyType.tp_str = (reprfunc)&FFuncs::Str;
+	PyType.tp_repr = (reprfunc)&FFuncs::Str;
 	PyType.tp_call = (ternaryfunc)&FFuncs::Call;
 	PyType.tp_descr_get = (descrgetfunc)&FFuncs::DescrGet;
 	PyType.tp_getattro = (getattrofunc)&PyObject_GenericGetAttr;
@@ -710,6 +712,7 @@ PyTypeObject InitializePyClassMethodWithClosureDescrType()
 	PyType.tp_traverse = (traverseproc)&FFuncs::Traverse;
 	PyType.tp_clear = (inquiry)&FFuncs::Clear;
 	PyType.tp_str = (reprfunc)&FFuncs::Str;
+	PyType.tp_repr = (reprfunc)&FFuncs::Str;
 	PyType.tp_call = (ternaryfunc)&FFuncs::Call;
 	PyType.tp_descr_get = (descrgetfunc)&FFuncs::DescrGet;
 	PyType.tp_getattro = (getattrofunc)&PyObject_GenericGetAttr;
diff --git a/Engine/Plugins/Experimental/PythonScriptPlugin/Source/PythonScriptPlugin/Private/PyWrapperDelegate.cpp b/Engine/Plugins/Experimental/PythonScriptPlugin/Source/PythonScriptPlugin/Private/PyWrapperDelegate.cpp
--- a/Engine/Plugins/Experimental/PythonScriptPlugin/Source/PythonScriptPlugin/Private/PyWrapperDelegate.cpp
+++ b/Engine/Plugins/Experimental/PythonScriptPlugin/Source/PythonScriptPlugin/Private/PyWrapperDelegate.cpp
@@ -774,6 +774,7 @@ PyTypeObject InitializePyWrapperDelegateType()
 	PyType.tp_dealloc = (destructor)&FFuncs::Dealloc;
 	PyType.tp_init = (initproc)&FFuncs::Init;
 	PyType.tp_str = (reprfunc)&FFuncs::Str;
+	PyType.tp_repr = (reprfunc)&FFuncs::Str;
 	PyType.tp_call = (ternaryfunc)&FFuncs::Call;
 
 	PyType.tp_methods = PyMethods;
@@ -1202,6 +1203,7 @@ PyTypeObject InitializePyWrapperMulticastDelegateType()
 	PyType.tp_dealloc = (destructor)&FFuncs::Dealloc;
 	PyType.tp_init = (initproc)&FFuncs::Init;
 	PyType.tp_str = (reprfunc)&FFuncs::Str;
+	PyType.tp_repr = (reprfunc)&FFuncs::Str;
 	PyType.tp_call = (ternaryfunc)&FFuncs::Call;
 
 	PyType.tp_methods = PyMethods;
diff --git a/Engine/Plugins/Experimental/PythonScriptPlugin/Source/PythonScriptPlugin/Private/PyWrapperEnum.cpp b/Engine/Plugins/Experimental/PythonScriptPlugin/Source/PythonScriptPlugin/Private/PyWrapperEnum.cpp
--- a/Engine/Plugins/Experimental/PythonScriptPlugin/Source/PythonScriptPlugin/Private/PyWrapperEnum.cpp
+++ b/Engine/Plugins/Experimental/PythonScriptPlugin/Source/PythonScriptPlugin/Private/PyWrapperEnum.cpp
@@ -428,6 +428,7 @@ PyTypeObject InitializePyWrapperEnumType()
 	PyType.tp_dealloc = (destructor)&FFuncs::Dealloc;
 	PyType.tp_init = (initproc)&FFuncs::Init;
 	PyType.tp_str = (reprfunc)&FFuncs::Str;
+	PyType.tp_repr = (reprfunc)&FFuncs::Str;
 	PyType.tp_richcompare = (richcmpfunc)&FFuncs::RichCmp;
 	PyType.tp_hash = (hashfunc)&FFuncs::Hash;
 
@@ -523,6 +524,7 @@ PyTypeObject InitializePyWrapperEnumValueDescrType()
 
 	PyType.tp_dealloc = (destructor)&FFuncs::Dealloc;
 	PyType.tp_str = (reprfunc)&FFuncs::Str;
+	PyType.tp_repr = (reprfunc)&FFuncs::Str;
 	PyType.tp_descr_get = (descrgetfunc)&FFuncs::DescrGet;
 	PyType.tp_descr_set = (descrsetfunc)&FFuncs::DescrSet;
 	PyType.tp_getattro = (getattrofunc)&PyObject_GenericGetAttr;
diff --git a/Engine/Plugins/Experimental/PythonScriptPlugin/Source/PythonScriptPlugin/Private/PyWrapperName.cpp b/Engine/Plugins/Experimental/PythonScriptPlugin/Source/PythonScriptPlugin/Private/PyWrapperName.cpp
--- a/Engine/Plugins/Experimental/PythonScriptPlugin/Source/PythonScriptPlugin/Private/PyWrapperName.cpp
+++ b/Engine/Plugins/Experimental/PythonScriptPlugin/Source/PythonScriptPlugin/Private/PyWrapperName.cpp
@@ -102,6 +102,11 @@ PyTypeObject InitializePyWrapperNameType()
 			return PyUnicode_FromString(TCHAR_TO_UTF8(*InSelf->Value.ToString()));
 		}
 
+		static PyObject* Repr(FPyWrapperName* InSelf)
+		{
+			return PyUnicode_FromString(TCHAR_TO_UTF8(*FString::Printf(TEXT("Name(\"%s\")"), *InSelf->Value.ToString())));
+		}
+
 		static PyObject* RichCmp(FPyWrapperName* InSelf, PyObject* InOther, int InOp)
 		{
 			FName Other;
@@ -171,6 +176,7 @@ PyTypeObject InitializePyWrapperNameType()
 
 	PyType.tp_init = (initproc)&FFuncs::Init;
 	PyType.tp_str = (reprfunc)&FFuncs::Str;
+	PyType.tp_repr = (reprfunc)&FFuncs::Repr;
 	PyType.tp_richcompare = (richcmpfunc)&FFuncs::RichCmp;
 	PyType.tp_hash = (hashfunc)&FFuncs::Hash;
 
diff --git a/Engine/Plugins/Experimental/PythonScriptPlugin/Source/PythonScriptPlugin/Private/PyWrapperObject.cpp b/Engine/Plugins/Experimental/PythonScriptPlugin/Source/PythonScriptPlugin/Private/PyWrapperObject.cpp
--- a/Engine/Plugins/Experimental/PythonScriptPlugin/Source/PythonScriptPlugin/Private/PyWrapperObject.cpp
+++ b/Engine/Plugins/Experimental/PythonScriptPlugin/Source/PythonScriptPlugin/Private/PyWrapperObject.cpp
@@ -1005,6 +1005,7 @@ PyTypeObject InitializePyWrapperObjectType()
 	PyType.tp_dealloc = (destructor)&FFuncs::Dealloc;
 	PyType.tp_init = (initproc)&FFuncs::Init;
 	PyType.tp_str = (reprfunc)&FFuncs::Str;
+	PyType.tp_repr = (reprfunc)&FFuncs::Str;
 	PyType.tp_hash = (hashfunc)&FFuncs::Hash;
 
 	PyType.tp_methods = PyMethods;
diff --git a/Engine/Plugins/Experimental/PythonScriptPlugin/Source/PythonScriptPlugin/Private/PyWrapperStruct.cpp b/Engine/Plugins/Experimental/PythonScriptPlugin/Source/PythonScriptPlugin/Private/PyWrapperStruct.cpp
--- a/Engine/Plugins/Experimental/PythonScriptPlugin/Source/PythonScriptPlugin/Private/PyWrapperStruct.cpp
+++ b/Engine/Plugins/Experimental/PythonScriptPlugin/Source/PythonScriptPlugin/Private/PyWrapperStruct.cpp
@@ -1218,6 +1218,7 @@ PyTypeObject InitializePyWrapperStructType()
 	PyType.tp_dealloc = (destructor)&FFuncs::Dealloc;
 	PyType.tp_init = (initproc)&FFuncs::Init;
 	PyType.tp_str = (reprfunc)&FFuncs::Str;
+	PyType.tp_repr = (reprfunc)&FFuncs::Str;
 	PyType.tp_richcompare = (richcmpfunc)&FFuncs::RichCmp;
 	PyType.tp_hash = (hashfunc)&FFuncs::Hash;
 
diff --git a/Engine/Plugins/Experimental/PythonScriptPlugin/Source/PythonScriptPlugin/Private/PyWrapperText.cpp b/Engine/Plugins/Experimental/PythonScriptPlugin/Source/PythonScriptPlugin/Private/PyWrapperText.cpp
--- a/Engine/Plugins/Experimental/PythonScriptPlugin/Source/PythonScriptPlugin/Private/PyWrapperText.cpp
+++ b/Engine/Plugins/Experimental/PythonScriptPlugin/Source/PythonScriptPlugin/Private/PyWrapperText.cpp
@@ -244,6 +244,11 @@ PyTypeObject InitializePyWrapperTextType()
 			return PyUnicode_FromString(TCHAR_TO_UTF8(*InSelf->Value.ToString()));
 		}
 
+		static PyObject* Repr(FPyWrapperText* InSelf)
+		{
+			return PyUnicode_FromString(TCHAR_TO_UTF8(*FString::Printf(TEXT("Text(\"%s\")"), *InSelf->Value.ToString())));
+		}
+
 		static PyObject* RichCmp(FPyWrapperText* InSelf, PyObject* InOther, int InOp)
 		{
 			FText Other;
@@ -469,6 +474,7 @@ PyTypeObject InitializePyWrapperTextType()
 
 	PyType.tp_init = (initproc)&FFuncs::Init;
 	PyType.tp_str = (reprfunc)&FFuncs::Str;
+	PyType.tp_repr = (reprfunc)&FFuncs::Repr;
 	PyType.tp_richcompare = (richcmpfunc)&FFuncs::RichCmp;
 	PyType.tp_hash = (hashfunc)&FFuncs::Hash;
 
