Added a guard when converting to FName in Python, to verify the name fits within NAME_SIZE

diff --git a/Engine/Plugins/Experimental/PythonScriptPlugin/Source/PythonScriptPlugin/Private/PyConversion.cpp b/Engine/Plugins/Experimental/PythonScriptPlugin/Source/PythonScriptPlugin/Private/PyConversion.cpp
--- a/Engine/Plugins/Experimental/PythonScriptPlugin/Source/PythonScriptPlugin/Private/PyConversion.cpp
+++ b/Engine/Plugins/Experimental/PythonScriptPlugin/Source/PythonScriptPlugin/Private/PyConversion.cpp
@@ -414,6 +414,11 @@ FPyConversionResult Nativize(PyObject* PyObj, FName& OutVal, const ESetErrorStat
 	FString NameStr;
 	if (Nativize(PyObj, NameStr, ESetErrorState::No))
 	{
+		if (NameStr.Len() >= NAME_SIZE)
+		{
+			PYCONVERSION_RETURN(FPyConversionResult::Failure(), TEXT("Nativize"), *FString::Printf(TEXT("Cannot nativize '%s' as 'Name' (lenth: %d, max length: %d)"), *PyUtil::GetFriendlyTypename(PyObj), NameStr.Len(), NAME_SIZE - 1));
+		}
+
 		OutVal = *NameStr;
 		return FPyConversionResult::SuccessWithCoercion();
 	}
