Minor fixes to prepare for the upcoming Python 3.9.7 update.
  - Tested those fixes with Python 3.7.x (current) and 3.9.7 (upcoming) on Windows and Mac

diff --git a/Engine/Plugins/Experimental/PythonScriptPlugin/Source/PythonScriptPlugin/Private/PyEngine.cpp b/Engine/Plugins/Experimental/PythonScriptPlugin/Source/PythonScriptPlugin/Private/PyEngine.cpp
--- a/Engine/Plugins/Experimental/PythonScriptPlugin/Source/PythonScriptPlugin/Private/PyEngine.cpp
+++ b/Engine/Plugins/Experimental/PythonScriptPlugin/Source/PythonScriptPlugin/Private/PyEngine.cpp
@@ -110,7 +110,7 @@ PyObject* GetBlueprintGeneratedTypes(PyObject* InSelf, PyObject* InArgs)
 #if PY_MAJOR_VERSION < 3
 				const bool bIsStringType = PyUnicode_Check(PyArg) || PyString_Check(PyArg);
 #else	// PY_MAJOR_VERSION < 3
-				const bool bIsStringType = PyUnicode_Check(PyArg);
+				const bool bIsStringType = PyUnicode_Check(PyArg) != 0;
 #endif	// PY_MAJOR_VERSION < 3
 				if (!bIsStringType && PyUtil::HasLength(PyArg))
 				{
diff --git a/Engine/Plugins/Experimental/PythonScriptPlugin/Source/PythonScriptPlugin/Private/PyUtil.cpp b/Engine/Plugins/Experimental/PythonScriptPlugin/Source/PythonScriptPlugin/Private/PyUtil.cpp
--- a/Engine/Plugins/Experimental/PythonScriptPlugin/Source/PythonScriptPlugin/Private/PyUtil.cpp
+++ b/Engine/Plugins/Experimental/PythonScriptPlugin/Source/PythonScriptPlugin/Private/PyUtil.cpp
@@ -1469,13 +1469,13 @@ bool FetchPythonError(FString& OutError)
 
 	PyErr_Clear();
 
-	// Raise the excepthook (if set)
-	// We set this to None after enabling our stderr redirection
+	// Raise the excepthook if it was changed.
 	{
-		PyObject* PyExceptHook = PySys_GetObject(PyCStrCast("excepthook"));
-		if (PyExceptHook && PyExceptHook != Py_None)
+		PyObject* PyDefaultExceptHook = PySys_GetObject(PyCStrCast("__excepthook__"));
+		PyObject* PyCurrentExceptHook = PySys_GetObject(PyCStrCast("excepthook"));
+		if (PyCurrentExceptHook && PyDefaultExceptHook && PyCurrentExceptHook != PyDefaultExceptHook)
 		{
-			FPyObjectPtr PyExceptHookResult = FPyObjectPtr::StealReference(PyObject_CallFunctionObjArgs(PyExceptHook, PyExceptionType.Get(), PyExceptionValue.Get(), PyExceptionTraceback.Get(), nullptr));
+			FPyObjectPtr PyExceptHookResult = FPyObjectPtr::StealReference(PyObject_CallFunctionObjArgs(PyCurrentExceptHook, PyExceptionType.Get(), PyExceptionValue.Get(), PyExceptionTraceback.Get(), nullptr));
 		}
 	}
 
diff --git a/Engine/Plugins/Experimental/PythonScriptPlugin/Source/PythonScriptPlugin/Private/PyWrapperObject.cpp b/Engine/Plugins/Experimental/PythonScriptPlugin/Source/PythonScriptPlugin/Private/PyWrapperObject.cpp
--- a/Engine/Plugins/Experimental/PythonScriptPlugin/Source/PythonScriptPlugin/Private/PyWrapperObject.cpp
+++ b/Engine/Plugins/Experimental/PythonScriptPlugin/Source/PythonScriptPlugin/Private/PyWrapperObject.cpp
@@ -1516,7 +1516,7 @@ public:
 			if (InPyFuncDef->FuncRetType && InPyFuncDef->FuncRetType != Py_None)
 			{
 				// If we have a tuple, then we actually want to return a bool but add every type within the tuple as output parameters
-				const bool bOptionalReturn = PyTuple_Check(InPyFuncDef->FuncRetType);
+				const bool bOptionalReturn = PyTuple_Check(InPyFuncDef->FuncRetType) != 0;
 
 				PyObject* RetType = bOptionalReturn ? (PyObject*)&PyBool_Type : InPyFuncDef->FuncRetType;
 				FProperty* RetProp = PyUtil::CreateProperty(RetType, 1, Func, TEXT("ReturnValue"));
diff --git a/Engine/Plugins/Experimental/PythonScriptPlugin/Source/PythonScriptPlugin/Private/PyWrapperText.cpp b/Engine/Plugins/Experimental/PythonScriptPlugin/Source/PythonScriptPlugin/Private/PyWrapperText.cpp
--- a/Engine/Plugins/Experimental/PythonScriptPlugin/Source/PythonScriptPlugin/Private/PyWrapperText.cpp
+++ b/Engine/Plugins/Experimental/PythonScriptPlugin/Source/PythonScriptPlugin/Private/PyWrapperText.cpp
@@ -71,7 +71,7 @@ bool ExtractFormatArguments(FPyWrapperText* InSelf, PyObject* InObj, const int32
 #if PY_MAJOR_VERSION < 3
 	const bool bIsStringType = PyUnicode_Check(InObj) || PyString_Check(InObj);
 #else	// PY_MAJOR_VERSION < 3
-	const bool bIsStringType = PyUnicode_Check(InObj);
+	const bool bIsStringType = PyUnicode_Check(InObj) != 0;
 #endif	// PY_MAJOR_VERSION < 3
 	if (!bIsStringType && PyUtil::HasLength(InObj))
 	{
diff --git a/Engine/Plugins/Experimental/PythonScriptPlugin/Source/PythonScriptPlugin/Private/PythonScriptPlugin.cpp b/Engine/Plugins/Experimental/PythonScriptPlugin/Source/PythonScriptPlugin/Private/PythonScriptPlugin.cpp
--- a/Engine/Plugins/Experimental/PythonScriptPlugin/Source/PythonScriptPlugin/Private/PythonScriptPlugin.cpp
+++ b/Engine/Plugins/Experimental/PythonScriptPlugin/Source/PythonScriptPlugin/Private/PythonScriptPlugin.cpp
@@ -684,12 +684,16 @@ void FPythonScriptPlugin::InitializePython()
 		Py_SetPythonHome(PyHomePath.GetData());
 		Py_InitializeEx(0); // 0 so Python doesn't override any signal handling
 
+#if PY_MAJOR_VERSION >= 3 && PY_MINOR_VERSION < 7
+		// NOTE: Since 3.7, those functions are called by Py_InitializeEx()
+		// 
 		// Ensure Python supports multiple threads via the GIL, as UE GC runs over multiple threads, 
 		// which may invoke FPyReferenceCollector::AddReferencedObjects on a background thread...
 		if (!PyEval_ThreadsInitialized())
 		{
 			PyEval_InitThreads();
 		}
+#endif // PY_MAJOR_VERSION >= 3 && PY_MINOR_VERSION < 7
 
 #if PLATFORM_WINDOWS && PY_MAJOR_VERSION >= 3
 		// We call _setmode here to restore the previous state
