Fixed static analyzer complaining about int/bool comparison in Python
  - Used static_cast<bool>() to work with both python 3.7 and 3.9. The macro expression type changed from bool to int between the versions.

diff --git a/Engine/Plugins/Experimental/PythonScriptPlugin/Source/PythonScriptPlugin/Private/PyEngine.cpp b/Engine/Plugins/Experimental/PythonScriptPlugin/Source/PythonScriptPlugin/Private/PyEngine.cpp
--- a/Engine/Plugins/Experimental/PythonScriptPlugin/Source/PythonScriptPlugin/Private/PyEngine.cpp
+++ b/Engine/Plugins/Experimental/PythonScriptPlugin/Source/PythonScriptPlugin/Private/PyEngine.cpp
@@ -110,7 +110,7 @@ PyObject* GetBlueprintGeneratedTypes(PyObject* InSelf, PyObject* InArgs)
 #if PY_MAJOR_VERSION < 3
 				const bool bIsStringType = PyUnicode_Check(PyArg) || PyString_Check(PyArg);
 #else	// PY_MAJOR_VERSION < 3
-				const bool bIsStringType = PyUnicode_Check(PyArg) != 0;
+				const bool bIsStringType = static_cast<bool>(PyUnicode_Check(PyArg));
 #endif	// PY_MAJOR_VERSION < 3
 				if (!bIsStringType && PyUtil::HasLength(PyArg))
 				{
diff --git a/Engine/Plugins/Experimental/PythonScriptPlugin/Source/PythonScriptPlugin/Private/PyWrapperObject.cpp b/Engine/Plugins/Experimental/PythonScriptPlugin/Source/PythonScriptPlugin/Private/PyWrapperObject.cpp
--- a/Engine/Plugins/Experimental/PythonScriptPlugin/Source/PythonScriptPlugin/Private/PyWrapperObject.cpp
+++ b/Engine/Plugins/Experimental/PythonScriptPlugin/Source/PythonScriptPlugin/Private/PyWrapperObject.cpp
@@ -1516,7 +1516,7 @@ public:
 			if (InPyFuncDef->FuncRetType && InPyFuncDef->FuncRetType != Py_None)
 			{
 				// If we have a tuple, then we actually want to return a bool but add every type within the tuple as output parameters
-				const bool bOptionalReturn = PyTuple_Check(InPyFuncDef->FuncRetType) != 0;
+				const bool bOptionalReturn = static_cast<bool>(PyTuple_Check(InPyFuncDef->FuncRetType));
 
 				PyObject* RetType = bOptionalReturn ? (PyObject*)&PyBool_Type : InPyFuncDef->FuncRetType;
 				FProperty* RetProp = PyUtil::CreateProperty(RetType, 1, Func, TEXT("ReturnValue"));
diff --git a/Engine/Plugins/Experimental/PythonScriptPlugin/Source/PythonScriptPlugin/Private/PyWrapperText.cpp b/Engine/Plugins/Experimental/PythonScriptPlugin/Source/PythonScriptPlugin/Private/PyWrapperText.cpp
--- a/Engine/Plugins/Experimental/PythonScriptPlugin/Source/PythonScriptPlugin/Private/PyWrapperText.cpp
+++ b/Engine/Plugins/Experimental/PythonScriptPlugin/Source/PythonScriptPlugin/Private/PyWrapperText.cpp
@@ -71,7 +71,7 @@ bool ExtractFormatArguments(FPyWrapperText* InSelf, PyObject* InObj, const int32
 #if PY_MAJOR_VERSION < 3
 	const bool bIsStringType = PyUnicode_Check(InObj) || PyString_Check(InObj);
 #else	// PY_MAJOR_VERSION < 3
-	const bool bIsStringType = PyUnicode_Check(InObj) != 0;
+	const bool bIsStringType = static_cast<bool>(PyUnicode_Check(InObj));
 #endif	// PY_MAJOR_VERSION < 3
 	if (!bIsStringType && PyUtil::HasLength(InObj))
 	{
