Python: Fixed an issue where remove_callable wouldn't work if the callable was a method (vs a function).
PyObject ptr equality doesn't work for method objects as there's a new instance everytime.
Instead, compare the callables tp_call pointers and PyMethod_Self values for methods.

diff --git a/Engine/Plugins/Experimental/PythonScriptPlugin/Source/PythonScriptPlugin/Private/PyWrapperDelegate.cpp b/Engine/Plugins/Experimental/PythonScriptPlugin/Source/PythonScriptPlugin/Private/PyWrapperDelegate.cpp
--- a/Engine/Plugins/Experimental/PythonScriptPlugin/Source/PythonScriptPlugin/Private/PyWrapperDelegate.cpp
+++ b/Engine/Plugins/Experimental/PythonScriptPlugin/Source/PythonScriptPlugin/Private/PyWrapperDelegate.cpp
@@ -90,6 +90,28 @@ void InitializePyWrapperDelegate(PyGenUtil::FNativePythonModule& ModuleInfo)
 namespace PyDelegateUtil
 {
 
+/**
+ * Tests for callable equality
+ * ob_type->tp_call is the function pointer that will be called, so they need to be the same.
+ * If the callable is a method, check that the self parameter is also the same.
+ * 
+ * Lhs & Rhs need to be callables!
+ */
+bool AreCallablesEqual(PyObject* Lhs, PyObject* Rhs)
+{
+	check(PyCallable_Check(Lhs) && PyCallable_Check(Rhs));
+
+	bool bAreCallablesEqual = (Lhs->ob_type == Rhs->ob_type); // Must be of the same type
+	bAreCallablesEqual = bAreCallablesEqual && (Lhs->ob_type->tp_call == Rhs->ob_type->tp_call); // Must be calling the same function/method
+
+	if (bAreCallablesEqual && PyMethod_Check(Lhs))
+	{
+		bAreCallablesEqual = (PyMethod_Self(Lhs) == PyMethod_Self(Rhs)); // Must have the same self
+	}
+
+	return bAreCallablesEqual;
+}
+
 bool PythonArgsToDelegate_ObjectAndName(PyObject* InArgs, const PyGenUtil::FGeneratedWrappedFunction& InDelegateSignature, FScriptDelegate& OutDelegate, const TCHAR* InFuncCtxt, const TCHAR* InErrorCtxt)
 {
 	PyObject* PyObj = nullptr;
@@ -991,7 +1013,7 @@ PyTypeObject InitializePyWrapperMulticastDelegateType()
 			{
 				if (const UPythonCallableForDelegate* PythonCallableForDelegate = ::Cast<UPythonCallableForDelegate>(DelegateObj))
 				{
-					if (PythonCallableForDelegate->GetCallable() == PyCallable)
+					if (PyDelegateUtil::AreCallablesEqual(PythonCallableForDelegate->GetCallable(), PyCallable))
 					{
 						bAddDelegate = false;
 						break;
@@ -1056,7 +1078,7 @@ PyTypeObject InitializePyWrapperMulticastDelegateType()
 			{
 				if (UPythonCallableForDelegate* PythonCallableForDelegate = ::Cast<UPythonCallableForDelegate>(DelegateObj))
 				{
-					if (PythonCallableForDelegate->GetCallable() == PyCallable)
+					if (PyDelegateUtil::AreCallablesEqual(PythonCallableForDelegate->GetCallable(), PyCallable))
 					{
 						PythonCallableForDelegateToRemove = PythonCallableForDelegate;
 						break;
@@ -1139,7 +1161,7 @@ PyTypeObject InitializePyWrapperMulticastDelegateType()
 			{
 				if (const UPythonCallableForDelegate* PythonCallableForDelegate = ::Cast<UPythonCallableForDelegate>(DelegateObj))
 				{
-					if (PythonCallableForDelegate->GetCallable() == PyCallable)
+					if (PyDelegateUtil::AreCallablesEqual(PythonCallableForDelegate->GetCallable(), PyCallable))
 					{
 						bContainsCallable = true;
 						break;
