Exposed UEnum display name to Python through method get_display_name(). Ex: print(unreal.TextureCompressionSettings.TC_GRAYSCALE.get_display_name())
  - Supported reading the enum entry display name set in C++ as well as the one injected from python with the [REDACTED_HANDLE] annotation.
  - Added unit tests.

diff --git a/Engine/Plugins/Experimental/PythonScriptPlugin/Source/PythonScriptPlugin/Private/PyTest.h b/Engine/Plugins/Experimental/PythonScriptPlugin/Source/PythonScriptPlugin/Private/PyTest.h
--- a/Engine/Plugins/Experimental/PythonScriptPlugin/Source/PythonScriptPlugin/Private/PyTest.h
+++ b/Engine/Plugins/Experimental/PythonScriptPlugin/Source/PythonScriptPlugin/Private/PyTest.h
@@ -23,7 +23,7 @@ DECLARE_DYNAMIC_MULTICAST_DELEGATE_OneParam(FPyTestMulticastDelegate, FString, I
 UENUM(BlueprintType)
 enum class EPyTestEnum : uint8
 {
-	One,
+	One UMETA(DisplayName = "Says One but my value is Zero"),
 	Two,
 };
 
diff --git a/Engine/Plugins/Experimental/PythonScriptPlugin/Source/PythonScriptPlugin/Private/PyWrapperEnum.cpp b/Engine/Plugins/Experimental/PythonScriptPlugin/Source/PythonScriptPlugin/Private/PyWrapperEnum.cpp
--- a/Engine/Plugins/Experimental/PythonScriptPlugin/Source/PythonScriptPlugin/Private/PyWrapperEnum.cpp
+++ b/Engine/Plugins/Experimental/PythonScriptPlugin/Source/PythonScriptPlugin/Private/PyWrapperEnum.cpp
@@ -8,6 +8,7 @@
 #include "PyConversion.h"
 
 #include "Containers/ArrayView.h"
+#include "Internationalization/Text.h"
 
 #if WITH_PYTHON
 
@@ -403,6 +404,17 @@ PyTypeObject InitializePyWrapperEnumType()
 			UEnum* Enum = FPyWrapperEnumMetaData::GetEnum(InType);
 			return PyConversion::Pythonize(Enum);
 		}
+
+		static PyObject* GetDisplayName(FPyWrapperEnum* InSelf)
+		{
+			if (!FPyWrapperEnum::ValidateInternalState(InSelf))
+			{
+				return nullptr;
+			}
+			int64 EnumEntryValue = FPyWrapperEnum::GetEnumEntryValue(InSelf);
+			UEnum* Enum = FPyWrapperEnumMetaData::GetEnum(InSelf);
+			return PyConversion::Pythonize(Enum->GetDisplayNameTextByValue(EnumEntryValue));
+		}
 	};
 
 	static PyMemberDef PyMembers[] = {
@@ -414,6 +426,7 @@ PyTypeObject InitializePyWrapperEnumType()
 	static PyMethodDef PyMethods[] = {
 		{ "cast", PyCFunctionCast(&FMethods::Cast), METH_VARARGS | METH_CLASS, "X.cast(object) -> enum -- cast the given object to this Unreal enum type" },
 		{ "static_enum", PyCFunctionCast(&FMethods::StaticEnum), METH_NOARGS | METH_CLASS, "X.static_enum() -> Enum -- get the Unreal enum of this type" },
+		{ "get_display_name", PyCFunctionCast(&FMethods::GetDisplayName), METH_NOARGS, "X.get_display_name() -> Text -- get the UMETA display name of this type in the current culture" },
 		{ nullptr, nullptr, 0, nullptr }
 	};
 
