Removed C++ code from PythonScriptPlugin refering to Python 2.x.

diff --git a/Engine/Plugins/Experimental/PythonScriptPlugin/Source/PythonScriptPlugin/Private/PyConversion.cpp b/Engine/Plugins/Experimental/PythonScriptPlugin/Source/PythonScriptPlugin/Private/PyConversion.cpp
--- a/Engine/Plugins/Experimental/PythonScriptPlugin/Source/PythonScriptPlugin/Private/PyConversion.cpp
+++ b/Engine/Plugins/Experimental/PythonScriptPlugin/Source/PythonScriptPlugin/Private/PyConversion.cpp
@@ -74,14 +74,6 @@ FPyConversionResult NativizeSigned(PyObject* PyObj, T& OutVal, const ESetErrorSt
 	// Booleans subclass integer, so exclude those explicitly
 	if (!PyBool_Check(PyObj))
 	{
-#if PY_MAJOR_VERSION < 3
-		if (PyInt_Check(PyObj))
-		{
-			OutVal = PyInt_AsLong(PyObj);
-			return FPyConversionResult::Success();
-		}
-#endif	// PY_MAJOR_VERSION < 3
-
 		if (PyLong_Check(PyObj))
 		{
 			OutVal = PyLong_AsLongLong(PyObj);
@@ -104,14 +96,6 @@ FPyConversionResult NativizeUnsigned(PyObject* PyObj, T& OutVal, const ESetError
 	// Booleans subclass integer, so exclude those explicitly
 	if (!PyBool_Check(PyObj))
 	{
-#if PY_MAJOR_VERSION < 3
-		if (PyInt_Check(PyObj))
-		{
-			OutVal = PyInt_AsSsize_t(PyObj);
-			return FPyConversionResult::Success();
-		}
-#endif	// PY_MAJOR_VERSION < 3
-
 		if (PyLong_Check(PyObj))
 		{
 			OutVal = PyLong_AsUnsignedLongLong(PyObj);
@@ -134,14 +118,6 @@ FPyConversionResult NativizeReal(PyObject* PyObj, T& OutVal, const ESetErrorStat
 	// Booleans subclass integer, so exclude those explicitly
 	if (!PyBool_Check(PyObj))
 	{
-#if PY_MAJOR_VERSION < 3
-		if (PyInt_Check(PyObj))
-		{
-			OutVal = PyInt_AsLong(PyObj);
-			return FPyConversionResult::SuccessWithCoercion();
-		}
-#endif	// PY_MAJOR_VERSION < 3
-
 		if (PyLong_Check(PyObj))
 		{
 			OutVal = PyLong_AsDouble(PyObj);
@@ -161,34 +137,14 @@ FPyConversionResult NativizeReal(PyObject* PyObj, T& OutVal, const ESetErrorStat
 template <typename T>
 FPyConversionResult PythonizeSigned(const T Val, PyObject*& OutPyObj, const ESetErrorState SetErrorState, const TCHAR* InErrorType)
 {
-#if PY_MAJOR_VERSION < 3
-	if (Val >= LONG_MIN && Val <= LONG_MAX)
-	{
-		OutPyObj = PyInt_FromLong(Val);
-	}
-	else
-#endif	// PY_MAJOR_VERSION < 3
-	{
-		OutPyObj = PyLong_FromLongLong(Val);
-	}
-
+	OutPyObj = PyLong_FromLongLong(Val);
 	return FPyConversionResult::Success();
 }
 
 template <typename T>
 FPyConversionResult PythonizeUnsigned(const T Val, PyObject*& OutPyObj, const ESetErrorState SetErrorState, const TCHAR* InErrorType)
 {
-#if PY_MAJOR_VERSION < 3
-	if (Val <= LONG_MAX)
-	{
-		OutPyObj = PyInt_FromSsize_t(Val);
-	}
-	else
-#endif	// PY_MAJOR_VERSION < 3
-	{
-		OutPyObj = PyLong_FromUnsignedLongLong(Val);
-	}
-
+	OutPyObj = PyLong_FromUnsignedLongLong(Val);
 	return FPyConversionResult::Success();
 }
 
@@ -220,14 +176,6 @@ FPyConversionResult Nativize(PyObject* PyObj, bool& OutVal, const ESetErrorState
 		OutVal = false;
 		return FPyConversionResult::Success();
 	}
-	
-#if PY_MAJOR_VERSION < 3
-	if (PyInt_Check(PyObj))
-	{
-		OutVal = PyInt_AsLong(PyObj) != 0;
-		return FPyConversionResult::SuccessWithCoercion();
-	}
-#endif	// PY_MAJOR_VERSION < 3
 
 	if (PyLong_Check(PyObj))
 	{
@@ -367,15 +315,6 @@ FPyConversionResult Nativize(PyObject* PyObj, FString& OutVal, const ESetErrorSt
 		}
 	}
 
-#if PY_MAJOR_VERSION < 3
-	if (PyString_Check(PyObj))
-	{
-		const char* PyUtf8Buffer = PyString_AsString(PyObj);
-		OutVal = UTF8_TO_TCHAR(PyUtf8Buffer);
-		return FPyConversionResult::Success();
-	}
-#endif	// PY_MAJOR_VERSION < 3
-
 	if (PyObject_IsInstance(PyObj, (PyObject*)&PyWrapperNameType) == 1)
 	{
 		FPyWrapperName* PyWrappedName = (FPyWrapperName*)PyObj;
@@ -388,17 +327,7 @@ FPyConversionResult Nativize(PyObject* PyObj, FString& OutVal, const ESetErrorSt
 
 FPyConversionResult Pythonize(const FString& Val, PyObject*& OutPyObj, const ESetErrorState SetErrorState)
 {
-#if PY_MAJOR_VERSION < 3
-	if (FCString::IsPureAnsi(*Val))
-	{
-		OutPyObj = PyString_FromStringAndSize(TCHAR_TO_ANSI(*Val), Val.Len());
-	}
-	else
-#endif	// PY_MAJOR_VERSION < 3
-	{
-		OutPyObj = PyUnicode_FromString(TCHAR_TO_UTF8(*Val));
-	}
-
+	OutPyObj = PyUnicode_FromString(TCHAR_TO_UTF8(*Val));
 	return FPyConversionResult::Success();
 }
 
@@ -469,13 +398,13 @@ FPyConversionResult Nativize(PyObject* PyObj, void*& OutVal, const ESetErrorStat
 #endif	// !(PY_MAJOR_VERSION >= 3 && PY_MINOR_VERSION >= 2)
 
 	// Capsule was added in Python 2.7 and 3.1
-#if (PY_MAJOR_VERSION >= 2 && PY_MINOR_VERSION >= 7) || (PY_MAJOR_VERSION >= 3 && PY_MINOR_VERSION >= 1)
+#if PY_MAJOR_VERSION >= 3 && PY_MINOR_VERSION >= 1
 	if (PyCapsule_CheckExact(PyObj))
 	{
 		OutVal = PyCapsule_GetPointer(PyObj, PyCapsule_GetName(PyObj));
 		return FPyConversionResult::Success();
 	}
-#endif	// (PY_MAJOR_VERSION >= 2 && PY_MINOR_VERSION >= 7) || (PY_MAJOR_VERSION >= 3 && PY_MINOR_VERSION >= 1)
+#endif	// PY_MAJOR_VERSION >= 3 && PY_MINOR_VERSION >= 1
 
 	if (PyObj == Py_None)
 	{
diff --git a/Engine/Plugins/Experimental/PythonScriptPlugin/Source/PythonScriptPlugin/Private/PyCore.cpp b/Engine/Plugins/Experimental/PythonScriptPlugin/Source/PythonScriptPlugin/Private/PyCore.cpp
--- a/Engine/Plugins/Experimental/PythonScriptPlugin/Source/PythonScriptPlugin/Private/PyCore.cpp
+++ b/Engine/Plugins/Experimental/PythonScriptPlugin/Source/PythonScriptPlugin/Private/PyCore.cpp
@@ -1807,13 +1807,8 @@ void InitializeModule()
 	PyGenUtil::FNativePythonModule NativePythonModule;
 	NativePythonModule.PyModuleMethods = PyCoreMethods;
 
-#if PY_MAJOR_VERSION >= 3
 	NativePythonModule.PyModule = PyImport_AddModule("_unreal_core");
 	PyModule_AddFunctions(NativePythonModule.PyModule, PyCoreMethods);
-#else	// PY_MAJOR_VERSION >= 3
-	NativePythonModule.PyModule = Py_InitModule("_unreal_core", PyCoreMethods);
-#endif	// PY_MAJOR_VERSION >= 3
-
 	PyType_Ready(&PyDelegateHandleType);
 
 	if (PyType_Ready(&PyScopedSlowTaskType) == 0)
diff --git a/Engine/Plugins/Experimental/PythonScriptPlugin/Source/PythonScriptPlugin/Private/PyEditor.cpp b/Engine/Plugins/Experimental/PythonScriptPlugin/Source/PythonScriptPlugin/Private/PyEditor.cpp
--- a/Engine/Plugins/Experimental/PythonScriptPlugin/Source/PythonScriptPlugin/Private/PyEditor.cpp
+++ b/Engine/Plugins/Experimental/PythonScriptPlugin/Source/PythonScriptPlugin/Private/PyEditor.cpp
@@ -212,12 +212,8 @@ void InitializeModule()
 	PyGenUtil::FNativePythonModule NativePythonModule;
 	NativePythonModule.PyModuleMethods = PyEditorMethods;
 
-#if PY_MAJOR_VERSION >= 3
 	NativePythonModule.PyModule = PyImport_AddModule("_unreal_editor");
 	PyModule_AddFunctions(NativePythonModule.PyModule, PyEditorMethods);
-#else	// PY_MAJOR_VERSION >= 3
-	NativePythonModule.PyModule = Py_InitModule("_unreal_editor", PyEditorMethods);
-#endif	// PY_MAJOR_VERSION >= 3
 
 	if (PyType_Ready(&PyScopedEditorTransactionType) == 0)
 	{
diff --git a/Engine/Plugins/Experimental/PythonScriptPlugin/Source/PythonScriptPlugin/Private/PyEngine.cpp b/Engine/Plugins/Experimental/PythonScriptPlugin/Source/PythonScriptPlugin/Private/PyEngine.cpp
--- a/Engine/Plugins/Experimental/PythonScriptPlugin/Source/PythonScriptPlugin/Private/PyEngine.cpp
+++ b/Engine/Plugins/Experimental/PythonScriptPlugin/Source/PythonScriptPlugin/Private/PyEngine.cpp
@@ -107,11 +107,7 @@ PyObject* GetBlueprintGeneratedTypes(PyObject* InSelf, PyObject* InArgs)
 			if (PyArg)
 			{
 				// Is this some kind of container, or a single value?
-#if PY_MAJOR_VERSION < 3
-				const bool bIsStringType = PyUnicode_Check(PyArg) || PyString_Check(PyArg);
-#else	// PY_MAJOR_VERSION < 3
 				const bool bIsStringType = static_cast<bool>(PyUnicode_Check(PyArg));
-#endif	// PY_MAJOR_VERSION < 3
 				if (!bIsStringType && PyUtil::HasLength(PyArg))
 				{
 					const Py_ssize_t SequenceLen = PyObject_Length(PyArg);
@@ -214,12 +210,8 @@ void InitializeModule()
 	PyGenUtil::FNativePythonModule NativePythonModule;
 	NativePythonModule.PyModuleMethods = PyEngineMethods;
 
-#if PY_MAJOR_VERSION >= 3
 	NativePythonModule.PyModule = PyImport_AddModule("_unreal_engine");
 	PyModule_AddFunctions(NativePythonModule.PyModule, PyEngineMethods);
-#else	// PY_MAJOR_VERSION >= 3
-	NativePythonModule.PyModule = Py_InitModule("_unreal_engine", PyEngineMethods);
-#endif	// PY_MAJOR_VERSION >= 3
 
 	if (PyType_Ready(&PyActorIteratorType) == 0)
 	{
diff --git a/Engine/Plugins/Experimental/PythonScriptPlugin/Source/PythonScriptPlugin/Private/PyGenUtil.cpp b/Engine/Plugins/Experimental/PythonScriptPlugin/Source/PythonScriptPlugin/Private/PyGenUtil.cpp
--- a/Engine/Plugins/Experimental/PythonScriptPlugin/Source/PythonScriptPlugin/Private/PyGenUtil.cpp
+++ b/Engine/Plugins/Experimental/PythonScriptPlugin/Source/PythonScriptPlugin/Private/PyGenUtil.cpp
@@ -154,18 +154,8 @@ void FGeneratedWrappedDynamicMethodsMixinBase::AddDynamicMethodImpl(FGeneratedWr
 
 const FGeneratedWrappedOperatorSignature& FGeneratedWrappedOperatorSignature::OpTypeToSignature(const EGeneratedWrappedOperatorType InOpType)
 {
-#if PY_MAJOR_VERSION >= 3
-	const TCHAR* BoolFuncName = TEXT("__bool__");
-	const TCHAR* DivideFuncName = TEXT("__truediv__");
-	const TCHAR* InlineDivideFuncName = TEXT("__truediv__");
-#else	// PY_MAJOR_VERSION >= 3
-	const TCHAR* BoolFuncName = TEXT("__nonzero__");
-	const TCHAR* DivideFuncName = TEXT("__div__");
-	const TCHAR* InlineDivideFuncName = TEXT("__idiv__");
-#endif	// PY_MAJOR_VERSION >= 3
-
 	static const FGeneratedWrappedOperatorSignature OperatorSignatures[(int32)EGeneratedWrappedOperatorType::Num] = {
-		FGeneratedWrappedOperatorSignature(EGeneratedWrappedOperatorType::Bool,				TEXT("bool"),	BoolFuncName,			EType::Bool,	EType::None),
+		FGeneratedWrappedOperatorSignature(EGeneratedWrappedOperatorType::Bool,				TEXT("bool"),	TEXT("__bool__"),		EType::Bool,	EType::None),
 		FGeneratedWrappedOperatorSignature(EGeneratedWrappedOperatorType::Equal,			TEXT("=="),		TEXT("__eq__"),			EType::Bool,	EType::Any),
 		FGeneratedWrappedOperatorSignature(EGeneratedWrappedOperatorType::NotEqual,			TEXT("!="),		TEXT("__ne__"),			EType::Bool,	EType::Any),
 		FGeneratedWrappedOperatorSignature(EGeneratedWrappedOperatorType::Less,				TEXT("<"),		TEXT("__lt__"),			EType::Bool,	EType::Any),
@@ -178,8 +168,8 @@ const FGeneratedWrappedOperatorSignature& FGeneratedWrappedOperatorSignature::Op
 		FGeneratedWrappedOperatorSignature(EGeneratedWrappedOperatorType::InlineSubtract,	TEXT("-="),		TEXT("__isub__"),		EType::Struct,	EType::Any),
 		FGeneratedWrappedOperatorSignature(EGeneratedWrappedOperatorType::Multiply,			TEXT("*"),		TEXT("__mul__"),		EType::Any,		EType::Any),
 		FGeneratedWrappedOperatorSignature(EGeneratedWrappedOperatorType::InlineMultiply,	TEXT("*="),		TEXT("__imul__"),		EType::Struct,	EType::Any),
-		FGeneratedWrappedOperatorSignature(EGeneratedWrappedOperatorType::Divide,			TEXT("/"),		DivideFuncName,			EType::Any,		EType::Any),
-		FGeneratedWrappedOperatorSignature(EGeneratedWrappedOperatorType::InlineDivide,		TEXT("/="),		InlineDivideFuncName,	EType::Struct,	EType::Any),
+		FGeneratedWrappedOperatorSignature(EGeneratedWrappedOperatorType::Divide,			TEXT("/"),		TEXT("__truediv__"),	EType::Any,		EType::Any),
+		FGeneratedWrappedOperatorSignature(EGeneratedWrappedOperatorType::InlineDivide,		TEXT("/="),		TEXT("__truediv__"),	EType::Struct,	EType::Any),
 		FGeneratedWrappedOperatorSignature(EGeneratedWrappedOperatorType::Modulus,			TEXT("%"),		TEXT("__mod__"),		EType::Any,		EType::Any),
 		FGeneratedWrappedOperatorSignature(EGeneratedWrappedOperatorType::InlineModulus,	TEXT("%="),		TEXT("__imod__"),		EType::Struct,	EType::Any),
 		FGeneratedWrappedOperatorSignature(EGeneratedWrappedOperatorType::And,				TEXT("&"),		TEXT("__and__"),		EType::Any,		EType::Any),
diff --git a/Engine/Plugins/Experimental/PythonScriptPlugin/Source/PythonScriptPlugin/Private/PyMethodWithClosure.cpp b/Engine/Plugins/Experimental/PythonScriptPlugin/Source/PythonScriptPlugin/Private/PyMethodWithClosure.cpp
--- a/Engine/Plugins/Experimental/PythonScriptPlugin/Source/PythonScriptPlugin/Private/PyMethodWithClosure.cpp
+++ b/Engine/Plugins/Experimental/PythonScriptPlugin/Source/PythonScriptPlugin/Private/PyMethodWithClosure.cpp
@@ -146,14 +146,6 @@ PyObject* FPyMethodWithClosureDef::Call(FPyMethodWithClosureDef* InDef, PyObject
 {
 	const int FuncFlags = InDef->MethodFlags & ~(METH_CLASS | METH_STATIC | METH_COEXIST);
 
-#if PY_MAJOR_VERSION < 3
-	if (FuncFlags == METH_OLDARGS)
-	{
-		PyErr_Format(PyExc_TypeError, "%s() METH_OLDARGS isn't supported", InDef->MethodName);
-		return nullptr;
-	}
-#endif	// PY_MAJOR_VERSION < 3
-
 	// Keywords take precedence
 	if (FuncFlags & METH_KEYWORDS)
 	{
@@ -446,14 +438,6 @@ PyTypeObject InitializePyCFunctionWithClosureType()
 
 		static PyObject* GetSelf(FPyCFunctionWithClosureObject* InSelf, void* InClosure)
 		{
-#if PY_MAJOR_VERSION < 3
-			if (PyEval_GetRestricted())
-			{
-				PyErr_SetString(PyExc_RuntimeError, "method.__self__ not accessible in restricted mode");
-				return nullptr;
-			}
-#endif	// PY_MAJOR_VERSION < 3
-
 			PyObject* Self = InSelf->SelfArg ? InSelf->SelfArg : Py_None;
 			Py_INCREF(Self);
 			return Self;
diff --git a/Engine/Plugins/Experimental/PythonScriptPlugin/Source/PythonScriptPlugin/Private/PySlate.cpp b/Engine/Plugins/Experimental/PythonScriptPlugin/Source/PythonScriptPlugin/Private/PySlate.cpp
--- a/Engine/Plugins/Experimental/PythonScriptPlugin/Source/PythonScriptPlugin/Private/PySlate.cpp
+++ b/Engine/Plugins/Experimental/PythonScriptPlugin/Source/PythonScriptPlugin/Private/PySlate.cpp
@@ -210,12 +210,8 @@ void InitializeModule()
 	PyGenUtil::FNativePythonModule NativePythonModule;
 	NativePythonModule.PyModuleMethods = PySlateMethods;
 
-#if PY_MAJOR_VERSION >= 3
 	NativePythonModule.PyModule = PyImport_AddModule("_unreal_slate");
 	PyModule_AddFunctions(NativePythonModule.PyModule, PySlateMethods);
-#else	// PY_MAJOR_VERSION >= 3
-	NativePythonModule.PyModule = Py_InitModule("_unreal_slate", PySlateMethods);
-#endif	// PY_MAJOR_VERSION >= 3
 
 	FPyWrapperTypeRegistry::Get().RegisterNativePythonModule(MoveTemp(NativePythonModule));
 }
diff --git a/Engine/Plugins/Experimental/PythonScriptPlugin/Source/PythonScriptPlugin/Private/PyUtil.cpp b/Engine/Plugins/Experimental/PythonScriptPlugin/Source/PythonScriptPlugin/Private/PyUtil.cpp
--- a/Engine/Plugins/Experimental/PythonScriptPlugin/Source/PythonScriptPlugin/Private/PyUtil.cpp
+++ b/Engine/Plugins/Experimental/PythonScriptPlugin/Source/PythonScriptPlugin/Private/PyUtil.cpp
@@ -58,31 +58,14 @@ FPyApiBuffer TCHARToPyApiBuffer(const TCHAR* InStr)
 
 FString PyObjectToUEString(PyObject* InPyObj)
 {
-	if (PyUnicode_Check(InPyObj)
-#if PY_MAJOR_VERSION < 3
-		|| PyString_Check(InPyObj)
-#endif	// PY_MAJOR_VERSION < 3
-		)
+	if (PyUnicode_Check(InPyObj))
 	{
 		return PyStringToUEString(InPyObj);
 	}
 
-#if PY_MAJOR_VERSION < 3
+	if (FPyObjectPtr PyStrObj = FPyObjectPtr::StealReference(PyObject_Str(InPyObj)))
 	{
-		FPyObjectPtr PyUnicodeObj = FPyObjectPtr::StealReference(PyObject_Unicode(InPyObj));
-		if (PyUnicodeObj)
-		{
-			return PyStringToUEString(PyUnicodeObj);
-		}
-	}
-#endif	// PY_MAJOR_VERSION < 3
-
-	{
-		FPyObjectPtr PyStrObj = FPyObjectPtr::StealReference(PyObject_Str(InPyObj));
-		if (PyStrObj)
-		{
-			return PyStringToUEString(PyStrObj);
-		}
+		return PyStringToUEString(PyStrObj);
 	}
 
 	return FString();
@@ -321,11 +304,7 @@ bool CalculatePropertyDef(PyTypeObject* InPyType, FPropertyDef& OutPropertyDef)
 		return true;
 	}
 
-	if (PyObject_IsSubclass((PyObject*)InPyType, (PyObject*)&PyUnicode_Type) == 1
-#if PY_MAJOR_VERSION < 3
-		|| PyObject_IsSubclass((PyObject*)InPyType, (PyObject*)&PyString_Type) == 1
-#endif	// PY_MAJOR_VERSION < 3
-		)
+	if (PyObject_IsSubclass((PyObject*)InPyType, (PyObject*)&PyUnicode_Type) == 1)
 	{
 		OutPropertyDef.PropertyClass = FStrProperty::StaticClass();
 		return true;
@@ -337,14 +316,6 @@ bool CalculatePropertyDef(PyTypeObject* InPyType, FPropertyDef& OutPropertyDef)
 		return true;
 	}
 
-#if PY_MAJOR_VERSION < 3
-	if (PyObject_IsSubclass((PyObject*)InPyType, (PyObject*)&PyInt_Type) == 1)
-	{
-		OutPropertyDef.PropertyClass = FIntProperty::StaticClass();
-		return true;
-	}
-#endif	// PY_MAJOR_VERSION < 3
-
 	if (PyObject_IsSubclass((PyObject*)InPyType, (PyObject*)&PyLong_Type) == 1)
 	{
 		OutPropertyDef.PropertyClass = FIntProperty::StaticClass();
@@ -599,11 +570,7 @@ bool InspectFunctionArgs(PyObject* InFunc, TArray<FString>& OutArgNames, TArray<
 	if (PyInspectModule)
 	{
 		PyObject* PyInspectDict = PyModule_GetDict(PyInspectModule);
-#if PY_MAJOR_VERSION >= 3
 		PyObject* PyGetArgSpecFunc = PyDict_GetItemString(PyInspectDict, "getfullargspec");
-#else	// PY_MAJOR_VERSION >= 3
-		PyObject* PyGetArgSpecFunc = PyDict_GetItemString(PyInspectDict, "getargspec");
-#endif	// PY_MAJOR_VERSION >= 3
 		if (PyGetArgSpecFunc)
 		{
 			FPyObjectPtr PyGetArgSpecResult = FPyObjectPtr::StealReference(PyObject_CallFunctionObjArgs(PyGetArgSpecFunc, InFunc, nullptr));
diff --git a/Engine/Plugins/Experimental/PythonScriptPlugin/Source/PythonScriptPlugin/Private/PyUtil.h b/Engine/Plugins/Experimental/PythonScriptPlugin/Source/PythonScriptPlugin/Private/PyUtil.h
--- a/Engine/Plugins/Experimental/PythonScriptPlugin/Source/PythonScriptPlugin/Private/PyUtil.h
+++ b/Engine/Plugins/Experimental/PythonScriptPlugin/Source/PythonScriptPlugin/Private/PyUtil.h
@@ -35,19 +35,11 @@ namespace PyUtil
 #endif	// PY_MAJOR_VERSION >= 3 && PY_MINOR_VERSION >= 2
 
 	/** Character type that this version of Python uses in its API */
-#if PY_MAJOR_VERSION >= 3
 	typedef wchar_t FPyApiChar;
-#else	// PY_MAJOR_VERSION >= 3
-	typedef char FPyApiChar;
-#endif	// PY_MAJOR_VERSION >= 3
 	typedef TArray<FPyApiChar> FPyApiBuffer;
 
 	/** Convert a TCHAR to a transient buffer that can be passed to a Python API that doesn't hold the result */
-#if PY_MAJOR_VERSION >= 3
 	#define TCHARToPyApiChar(InStr) TCHAR_TO_WCHAR(InStr)
-#else	// PY_MAJOR_VERSION >= 3
-	#define TCHARToPyApiChar(InStr) TCHAR_TO_UTF8(InStr)
-#endif	// PY_MAJOR_VERSION >= 3
 
 	extern const FName DefaultPythonPropertyName;
 
diff --git a/Engine/Plugins/Experimental/PythonScriptPlugin/Source/PythonScriptPlugin/Private/PyWrapperArray.cpp b/Engine/Plugins/Experimental/PythonScriptPlugin/Source/PythonScriptPlugin/Private/PyWrapperArray.cpp
--- a/Engine/Plugins/Experimental/PythonScriptPlugin/Source/PythonScriptPlugin/Private/PyWrapperArray.cpp
+++ b/Engine/Plugins/Experimental/PythonScriptPlugin/Source/PythonScriptPlugin/Private/PyWrapperArray.cpp
@@ -760,11 +760,7 @@ int FPyWrapperArray::Reverse(FPyWrapperArray* InSelf)
 	return 0;
 }
 
-#if PY_MAJOR_VERSION < 3
-int FPyWrapperArray::Sort(FPyWrapperArray* InSelf, PyObject* InCmp, PyObject* InKey, bool InReverse)
-#else	// PY_MAJOR_VERSION < 3
 int FPyWrapperArray::Sort(FPyWrapperArray* InSelf, PyObject* InKey, bool InReverse)
-#endif	// PY_MAJOR_VERSION < 3
 {
 	if (!ValidateInternalState(InSelf))
 	{
@@ -796,9 +792,7 @@ int FPyWrapperArray::Sort(FPyWrapperArray* InSelf, PyObject* InKey, bool InRever
 	FPyObjectPtr PyListSortFunc = FPyObjectPtr::StealReference(PyObject_GetAttrString(PyList, "sort"));
 	FPyObjectPtr PyListSortArgs = FPyObjectPtr::StealReference(PyTuple_New(0));
 	FPyObjectPtr PyListSortKwds = FPyObjectPtr::StealReference(PyDict_New());
-#if PY_MAJOR_VERSION < 3
-	PyDict_SetItemString(PyListSortKwds, "cmp", InCmp ? InCmp : Py_None);
-#endif	// PY_MAJOR_VERSION < 3
+
 	PyDict_SetItemString(PyListSortKwds, "key", InKey ? InKey : Py_None);
 	PyDict_SetItemString(PyListSortKwds, "reverse", InReverse ? Py_True : Py_False);
 
@@ -1347,19 +1341,11 @@ PyTypeObject InitializePyWrapperArrayType()
 
 		static PyObject* Sort(FPyWrapperArray* InSelf, PyObject* InArgs, PyObject* InKwds)
 		{
-#if PY_MAJOR_VERSION < 3
-			PyObject* PyCmpObject = nullptr;
-#endif	// PY_MAJOR_VERSION < 3
 			PyObject* PyKeyObj = nullptr;
 			PyObject* PyReverseObj = nullptr;
 
-#if PY_MAJOR_VERSION < 3
-			static const char *ArgsKwdList[] = { "cmp", "key", "reverse", nullptr };
-			if (!PyArg_ParseTupleAndKeywords(InArgs, InKwds, "|OOO:sort", (char**)ArgsKwdList, &PyCmpObject, &PyKeyObj, &PyReverseObj))
-#else	// PY_MAJOR_VERSION < 3
 			static const char *ArgsKwdList[] = { "key", "reverse", nullptr };
 			if (!PyArg_ParseTupleAndKeywords(InArgs, InKwds, "|OO:sort", (char**)ArgsKwdList, &PyKeyObj, &PyReverseObj))
-#endif	// PY_MAJOR_VERSION < 3
 			{
 				return nullptr;
 			}
@@ -1371,11 +1357,7 @@ PyTypeObject InitializePyWrapperArrayType()
 				return nullptr;
 			}
 
-#if PY_MAJOR_VERSION < 3
-			if (FPyWrapperArray::Sort(InSelf, PyCmpObject, PyKeyObj, bReverse) != 0)
-#else	// PY_MAJOR_VERSION < 3
 			if (FPyWrapperArray::Sort(InSelf, PyKeyObj, bReverse) != 0)
-#endif	// PY_MAJOR_VERSION < 3
 			{
 				return nullptr;
 			}
@@ -1418,11 +1400,7 @@ PyTypeObject InitializePyWrapperArrayType()
 		{ "pop", PyCFunctionCast(&FMethods::Pop), METH_VARARGS, "x.pop(index=len-1) -> value -- remove and return the value at the given index in this Unreal array, or raise IndexError if the index is out-of-bounds" },
 		{ "remove", PyCFunctionCast(&FMethods::Remove), METH_VARARGS, "x.remove(value) -> None -- remove the first matching value in this Unreal array, or raise ValueError if missing" },
 		{ "reverse", PyCFunctionCast(&FMethods::Reverse), METH_NOARGS, "x.reverse() -> None -- reverse this Unreal array in-place" },
-#if PY_MAJOR_VERSION < 3
-		{ "sort", PyCFunctionCast(&FMethods::Sort), METH_VARARGS | METH_KEYWORDS, "x.sort(cmp=None, key=None, reverse=False) -> None -- stable sort this Unreal array in-place" },
-#else	// PY_MAJOR_VERSION < 3
 		{ "sort", PyCFunctionCast(&FMethods::Sort), METH_VARARGS | METH_KEYWORDS, "x.sort(key=None, reverse=False) -> None -- stable sort this Unreal array in-place" },
-#endif	// PY_MAJOR_VERSION < 3
 		{ "resize", PyCFunctionCast(&FMethods::Resize), METH_VARARGS, "x.resize(len) -> None -- resize this Unreal array to hold the given number of elements" },
 		{ nullptr, nullptr, 0, nullptr }
 	};
diff --git a/Engine/Plugins/Experimental/PythonScriptPlugin/Source/PythonScriptPlugin/Private/PyWrapperArray.h b/Engine/Plugins/Experimental/PythonScriptPlugin/Source/PythonScriptPlugin/Private/PyWrapperArray.h
--- a/Engine/Plugins/Experimental/PythonScriptPlugin/Source/PythonScriptPlugin/Private/PyWrapperArray.h
+++ b/Engine/Plugins/Experimental/PythonScriptPlugin/Source/PythonScriptPlugin/Private/PyWrapperArray.h
@@ -96,11 +96,7 @@ struct FPyWrapperArray : public FPyWrapperBase
 	static int Reverse(FPyWrapperArray* InSelf);
 
 	/** Sort this container (equivalent to 'x.sort()' in Python) */
-#if PY_MAJOR_VERSION < 3
-	static int Sort(FPyWrapperArray* InSelf, PyObject* InCmp, PyObject* InKey, bool InReverse);
-#else	// PY_MAJOR_VERSION < 3
 	static int Sort(FPyWrapperArray* InSelf, PyObject* InKey, bool InReverse);
-#endif	// PY_MAJOR_VERSION < 3
 
 	/** Resize this container (equivalent to 'x.resize(l)' in Python) */
 	static int Resize(FPyWrapperArray* InSelf, Py_ssize_t InLen);
diff --git a/Engine/Plugins/Experimental/PythonScriptPlugin/Source/PythonScriptPlugin/Private/PyWrapperDelegate.cpp b/Engine/Plugins/Experimental/PythonScriptPlugin/Source/PythonScriptPlugin/Private/PyWrapperDelegate.cpp
--- a/Engine/Plugins/Experimental/PythonScriptPlugin/Source/PythonScriptPlugin/Private/PyWrapperDelegate.cpp
+++ b/Engine/Plugins/Experimental/PythonScriptPlugin/Source/PythonScriptPlugin/Private/PyWrapperDelegate.cpp
@@ -805,12 +805,7 @@ PyTypeObject InitializePyWrapperDelegateType()
 	PyType.tp_doc = "Type for all Unreal exposed delegate instances";
 
 	static PyNumberMethods PyNumber;
-#if PY_MAJOR_VERSION >= 3
 	PyNumber.nb_bool = (inquiry)&FNumberFuncs::Bool;
-#else	// PY_MAJOR_VERSION >= 3
-	PyNumber.nb_nonzero = (inquiry)&FNumberFuncs::Bool;
-#endif	// PY_MAJOR_VERSION >= 3
-
 	PyType.tp_as_number = &PyNumber;
 
 	return PyType;
@@ -1234,12 +1229,7 @@ PyTypeObject InitializePyWrapperMulticastDelegateType()
 	PyType.tp_doc = "Type for all Unreal exposed multicast delegate instances";
 
 	static PyNumberMethods PyNumber;
-#if PY_MAJOR_VERSION >= 3
 	PyNumber.nb_bool = (inquiry)&FNumberFuncs::Bool;
-#else	// PY_MAJOR_VERSION >= 3
-	PyNumber.nb_nonzero = (inquiry)&FNumberFuncs::Bool;
-#endif	// PY_MAJOR_VERSION >= 3
-
 	PyType.tp_as_number = &PyNumber;
 
 	return PyType;
diff --git a/Engine/Plugins/Experimental/PythonScriptPlugin/Source/PythonScriptPlugin/Private/PyWrapperMap.cpp b/Engine/Plugins/Experimental/PythonScriptPlugin/Source/PythonScriptPlugin/Private/PyWrapperMap.cpp
--- a/Engine/Plugins/Experimental/PythonScriptPlugin/Source/PythonScriptPlugin/Private/PyWrapperMap.cpp
+++ b/Engine/Plugins/Experimental/PythonScriptPlugin/Source/PythonScriptPlugin/Private/PyWrapperMap.cpp
@@ -1474,77 +1474,6 @@ PyTypeObject InitializePyWrapperMapType()
 			Py_RETURN_NONE;
 		}
 
-#if PY_MAJOR_VERSION < 3
-
-		static PyObject* HasKey(FPyWrapperMap* InSelf, PyObject* InArgs)
-		{
-			PyObject* PyObj = nullptr;
-			if (!PyArg_ParseTuple(InArgs, "|O:has_key", &PyObj))
-			{
-				return nullptr;
-			}
-
-			const int ContainsResult = FPyWrapperMap::Contains(InSelf, PyObj);
-			if (ContainsResult == -1)
-			{
-				return nullptr;
-			}
-
-			if (ContainsResult == 1)
-			{
-				Py_RETURN_TRUE;
-			}
-
-			Py_RETURN_FALSE;
-		}
-
-		static PyObject* IterItems(FPyWrapperMap* InSelf)
-		{
-			return MakeMapIter<FPyWrapperMapItemIterator>(InSelf);
-		}
-
-		static PyObject* IterKeys(FPyWrapperMap* InSelf)
-		{
-			return MakeMapIter<FPyWrapperMapKeyIterator>(InSelf);
-		}
-
-		static PyObject* IterValues(FPyWrapperMap* InSelf)
-		{
-			return MakeMapIter<FPyWrapperMapValueIterator>(InSelf);
-		}
-
-		static PyObject* Items(FPyWrapperMap* InSelf)
-		{
-			return FPyWrapperMap::Items(InSelf);
-		}
-
-		static PyObject* Keys(FPyWrapperMap* InSelf)
-		{
-			return FPyWrapperMap::Keys(InSelf);
-		}
-
-		static PyObject* Values(FPyWrapperMap* InSelf)
-		{
-			return FPyWrapperMap::Values(InSelf);
-		}
-
-		static PyObject* ViewItems(FPyWrapperMap* InSelf)
-		{
-			return MakeMapView<FPyWrapperMapItemView>(InSelf);
-		}
-
-		static PyObject* ViewKeys(FPyWrapperMap* InSelf)
-		{
-			return MakeMapView<FPyWrapperMapKeyView>(InSelf);
-		}
-
-		static PyObject* ViewValues(FPyWrapperMap* InSelf)
-		{
-			return MakeMapView<FPyWrapperMapValueView>(InSelf);
-		}
-
-#else	// PY_MAJOR_VERSION < 3
-
 		static PyObject* Items(FPyWrapperMap* InSelf)
 		{
 			return MakeMapView<FPyWrapperMapItemView>(InSelf);
@@ -1559,8 +1488,6 @@ PyTypeObject InitializePyWrapperMapType()
 		{
 			return MakeMapView<FPyWrapperMapValueView>(InSelf);
 		}
-
-#endif	// PY_MAJOR_VERSION < 3
 	};
 
 	static PyMethodDef PyMethods[] = {
@@ -1574,22 +1501,9 @@ PyTypeObject InitializePyWrapperMapType()
 		{ "pop", PyCFunctionCast(&FMethods::Pop), METH_VARARGS | METH_KEYWORDS, "x.pop(key, default=None) -> value -- remove key and return its value, or default if key not present, or raise KeyError if no default" },
 		{ "popitem", PyCFunctionCast(&FMethods::PopItem), METH_NOARGS, "x.popitem() -> pair -- remove and return an arbitrary pair from this Unreal map, or raise KeyError if the map is empty" },
 		{ "update", PyCFunctionCast(&FMethods::Update), METH_VARARGS | METH_KEYWORDS, "x.update(...) -> None -- update this Unreal map from the given mapping or sequence pairs type or key->value arguments" },
-#if PY_MAJOR_VERSION < 3
-		{ "has_key", PyCFunctionCast(&FMethods::HasKey), METH_VARARGS, "x.has_key(k) -> bool -- does this Unreal map contain the given key? (equivalent to k in x)" },
-		{ "iteritems", PyCFunctionCast(&FMethods::IterItems), METH_NOARGS, "x.iteritems() -> iter -- an iterator over the key->value pairs of this Unreal map" },
-		{ "iterkeys", PyCFunctionCast(&FMethods::IterKeys), METH_NOARGS, "x.iterkeys() -> iter -- an iterator over the keys of this Unreal map" },
-		{ "itervalues", PyCFunctionCast(&FMethods::IterValues), METH_NOARGS, "x.itervalues() -> iter -- an iterator over the values of this Unreal map" },
-		{ "items", PyCFunctionCast(&FMethods::Items), METH_NOARGS, "x.items() -> iter -- a Python list containing the key->value pairs of this Unreal map" },
-		{ "keys", PyCFunctionCast(&FMethods::Keys), METH_NOARGS, "x.keys() -> iter -- a Python list containing the keys of this Unreal map" },
-		{ "values", PyCFunctionCast(&FMethods::Values), METH_NOARGS, "x.values() -> iter -- a Python list containing the values of this Unreal map" },
-		{ "viewitems", PyCFunctionCast(&FMethods::ViewItems), METH_NOARGS, "x.viewitems() -> view -- a set-like view of the key->value pairs of this Unreal map" },
-		{ "viewkeys", PyCFunctionCast(&FMethods::ViewKeys), METH_NOARGS, "x.viewkeys() -> view -- a set-like view of the keys of this Unreal map" },
-		{ "viewvalues", PyCFunctionCast(&FMethods::ViewValues), METH_NOARGS, "x.viewvalues() -> view -- a view of the values of this Unreal map" },
-#else	// PY_MAJOR_VERSION < 3
 		{ "items", PyCFunctionCast(&FMethods::Items), METH_NOARGS, "x.items() -> view -- a set-like view of the key->value pairs of this Unreal map" },
 		{ "keys", PyCFunctionCast(&FMethods::Keys), METH_NOARGS, "x.keys() -> view -- a set-like view of the keys of this Unreal map" },
 		{ "values", PyCFunctionCast(&FMethods::Values), METH_NOARGS, "x.values() -> view -- a view of the values of this Unreal map" },
-#endif	// PY_MAJOR_VERSION < 3
 		{ nullptr, nullptr, 0, nullptr }
 	};
 
@@ -1824,10 +1738,6 @@ PyTypeObject InitializePyWrapperMapSetViewType(const char* InName)
 
 	PyTypeObject PyType = InitializePyWrapperMapViewType<TImpl>(InName);
 
-#if PY_MAJOR_VERSION < 3
-	PyType.tp_flags |= Py_TPFLAGS_CHECKTYPES;
-#endif	// PY_MAJOR_VERSION < 3
-
 	static PyNumberMethods PyNumber;
 	PyNumber.nb_subtract = (binaryfunc)&FNumberFuncs::Sub;
 	PyNumber.nb_and = (binaryfunc)&FNumberFuncs::And;
diff --git a/Engine/Plugins/Experimental/PythonScriptPlugin/Source/PythonScriptPlugin/Private/PyWrapperSet.cpp b/Engine/Plugins/Experimental/PythonScriptPlugin/Source/PythonScriptPlugin/Private/PyWrapperSet.cpp
--- a/Engine/Plugins/Experimental/PythonScriptPlugin/Source/PythonScriptPlugin/Private/PyWrapperSet.cpp
+++ b/Engine/Plugins/Experimental/PythonScriptPlugin/Source/PythonScriptPlugin/Private/PyWrapperSet.cpp
@@ -1515,9 +1515,6 @@ PyTypeObject InitializePyWrapperSetType()
 	PyType.tp_methods = PyMethods;
 
 	PyType.tp_flags = Py_TPFLAGS_DEFAULT;
-#if PY_MAJOR_VERSION < 3
-	PyType.tp_flags |= Py_TPFLAGS_CHECKTYPES;
-#endif	// PY_MAJOR_VERSION < 3
 	PyType.tp_doc = "Type for all Unreal exposed set instances";
 
 	static PyNumberMethods PyNumber;
diff --git a/Engine/Plugins/Experimental/PythonScriptPlugin/Source/PythonScriptPlugin/Private/PyWrapperStruct.cpp b/Engine/Plugins/Experimental/PythonScriptPlugin/Source/PythonScriptPlugin/Private/PyWrapperStruct.cpp
--- a/Engine/Plugins/Experimental/PythonScriptPlugin/Source/PythonScriptPlugin/Private/PyWrapperStruct.cpp
+++ b/Engine/Plugins/Experimental/PythonScriptPlugin/Source/PythonScriptPlugin/Private/PyWrapperStruct.cpp
@@ -1225,30 +1225,18 @@ PyTypeObject InitializePyWrapperStructType()
 	PyType.tp_methods = PyMethods;
 
 	PyType.tp_flags = Py_TPFLAGS_DEFAULT | Py_TPFLAGS_BASETYPE;
-#if PY_MAJOR_VERSION < 3
-	PyType.tp_flags |= Py_TPFLAGS_CHECKTYPES;
-#endif	// PY_MAJOR_VERSION < 3
 	PyType.tp_doc = "Type for all Unreal exposed struct instances";
 
 	static PyNumberMethods PyNumber;
-#if PY_MAJOR_VERSION >= 3
 	PyNumber.nb_bool = (inquiry)&FNumberFuncs::Bool;
-#else	// PY_MAJOR_VERSION >= 3
-	PyNumber.nb_nonzero = (inquiry)&FNumberFuncs::Bool;
-#endif	// PY_MAJOR_VERSION >= 3
 	PyNumber.nb_add = (binaryfunc)&FNumberFuncs::Add;
 	PyNumber.nb_inplace_add = (binaryfunc)&FNumberFuncs::InlineAdd;
 	PyNumber.nb_subtract = (binaryfunc)&FNumberFuncs::Subtract;
 	PyNumber.nb_inplace_subtract = (binaryfunc)&FNumberFuncs::InlineSubtract;
 	PyNumber.nb_multiply = (binaryfunc)&FNumberFuncs::Multiply;
 	PyNumber.nb_inplace_multiply = (binaryfunc)&FNumberFuncs::InlineMultiply;
-#if PY_MAJOR_VERSION >= 3
 	PyNumber.nb_true_divide = (binaryfunc)&FNumberFuncs::Divide;
 	PyNumber.nb_inplace_true_divide = (binaryfunc)&FNumberFuncs::InlineDivide;
-#else	// PY_MAJOR_VERSION >= 3
-	PyNumber.nb_divide = (binaryfunc)&FNumberFuncs::Divide;
-	PyNumber.nb_inplace_divide = (binaryfunc)&FNumberFuncs::InlineDivide;
-#endif	// PY_MAJOR_VERSION >= 3
 	PyNumber.nb_remainder = (binaryfunc)&FNumberFuncs::Modulus;
 	PyNumber.nb_inplace_remainder = (binaryfunc)&FNumberFuncs::InlineModulus;
 	PyNumber.nb_and = (binaryfunc)&FNumberFuncs::And;
diff --git a/Engine/Plugins/Experimental/PythonScriptPlugin/Source/PythonScriptPlugin/Private/PyWrapperText.cpp b/Engine/Plugins/Experimental/PythonScriptPlugin/Source/PythonScriptPlugin/Private/PyWrapperText.cpp
--- a/Engine/Plugins/Experimental/PythonScriptPlugin/Source/PythonScriptPlugin/Private/PyWrapperText.cpp
+++ b/Engine/Plugins/Experimental/PythonScriptPlugin/Source/PythonScriptPlugin/Private/PyWrapperText.cpp
@@ -38,15 +38,6 @@ bool ExtractFormatArgumentValue(FPyWrapperText* InSelf, PyObject* InObj, FFormat
 	}
 
 	// Don't use PyConversion for numeric types as they would allow coercion from float<->int
-#if PY_MAJOR_VERSION < 3
-	if (PyInt_Check(InObj))
-	{
-		OutFormatArg.ArgumentValueType = EFormatArgumentType::Int;
-		OutFormatArg.ArgumentValueInt = PyInt_AsLong(InObj);
-		return true;
-	}
-#endif	// PY_MAJOR_VERSION < 3
-
 	if (PyLong_Check(InObj))
 	{
 		OutFormatArg.ArgumentValueType = EFormatArgumentType::Int;
@@ -68,11 +59,7 @@ bool ExtractFormatArgumentValue(FPyWrapperText* InSelf, PyObject* InObj, FFormat
 bool ExtractFormatArguments(FPyWrapperText* InSelf, PyObject* InObj, const int32 InArgIndex, TArray<FFormatArgumentData>& InOutFormatArgs)
 {
 	// Is this some kind of container, or a single value?
-#if PY_MAJOR_VERSION < 3
-	const bool bIsStringType = PyUnicode_Check(InObj) || PyString_Check(InObj);
-#else	// PY_MAJOR_VERSION < 3
 	const bool bIsStringType = static_cast<bool>(PyUnicode_Check(InObj));
-#endif	// PY_MAJOR_VERSION < 3
 	if (!bIsStringType && PyUtil::HasLength(InObj))
 	{
 		const Py_ssize_t SequenceLen = PyObject_Length(InObj);
diff --git a/Engine/Plugins/Experimental/PythonScriptPlugin/Source/PythonScriptPlugin/Private/PythonScriptPlugin.cpp b/Engine/Plugins/Experimental/PythonScriptPlugin/Source/PythonScriptPlugin/Private/PythonScriptPlugin.cpp
--- a/Engine/Plugins/Experimental/PythonScriptPlugin/Source/PythonScriptPlugin/Private/PythonScriptPlugin.cpp
+++ b/Engine/Plugins/Experimental/PythonScriptPlugin/Source/PythonScriptPlugin/Private/PythonScriptPlugin.cpp
@@ -720,11 +720,12 @@ void FPythonScriptPlugin::InitializePython()
 
 	// Initialize the Python interpreter
 	{
+		static_assert(PY_MAJOR_VERSION >= 3, "Unreal Engine Python integration doesn't support versions prior to Python 3.x");
 		UE_LOG(LogPython, Log, TEXT("Using Python %d.%d.%d"), PY_MAJOR_VERSION, PY_MINOR_VERSION, PY_MICRO_VERSION);
 
 		// Python 3 changes the console mode from O_TEXT to O_BINARY which affects other uses of the console
 		// So change the console mode back to its current setting after Py_Initialize has been called
-#if PLATFORM_WINDOWS && PY_MAJOR_VERSION >= 3
+#if PLATFORM_WINDOWS
 		// We call _setmode here to cache the current state
 		CA_SUPPRESS(6031)
 		fflush(stdin);
@@ -735,7 +736,7 @@ void FPythonScriptPlugin::InitializePython()
 		CA_SUPPRESS(6031)
 		fflush(stderr);
 		const int StdErrMode = _setmode(_fileno(stderr), _O_TEXT);
-#endif	// PLATFORM_WINDOWS && PY_MAJOR_VERSION >= 3
+#endif	// PLATFORM_WINDOWS
 
 #if PY_MAJOR_VERSION >= 3 && PY_MINOR_VERSION >= 7
 		// Python 3.7+ changes the C locale which affects functions using C string APIs
@@ -770,7 +771,7 @@ void FPythonScriptPlugin::InitializePython()
 		}
 #endif // PY_MAJOR_VERSION >= 3 && PY_MINOR_VERSION < 7
 
-#if PLATFORM_WINDOWS && PY_MAJOR_VERSION >= 3
+#if PLATFORM_WINDOWS
 		// We call _setmode here to restore the previous state
 		if (StdInMode != -1)
 		{
@@ -793,7 +794,7 @@ void FPythonScriptPlugin::InitializePython()
 			CA_SUPPRESS(6031)
 			_setmode(_fileno(stderr), StdErrMode);
 		}
-#endif	// PLATFORM_WINDOWS && PY_MAJOR_VERSION >= 3
+#endif	// PLATFORM_WINDOWS
 
 #if PY_MAJOR_VERSION >= 3 && PY_MINOR_VERSION >= 7
 		// We call setlocale here to restore the previous state
