Prevented adding non-existing directories to Python 'sys.path'
  - This makes Python a bit faster at startup (according to my benchmark) and reduce noise looking at the sys.path.

diff --git a/Engine/Plugins/Experimental/PythonScriptPlugin/Source/PythonScriptPlugin/Private/PyUtil.cpp b/Engine/Plugins/Experimental/PythonScriptPlugin/Source/PythonScriptPlugin/Private/PyUtil.cpp
--- a/Engine/Plugins/Experimental/PythonScriptPlugin/Source/PythonScriptPlugin/Private/PyUtil.cpp
+++ b/Engine/Plugins/Experimental/PythonScriptPlugin/Source/PythonScriptPlugin/Private/PyUtil.cpp
@@ -1071,6 +1071,11 @@ void AddSitePackagesPath(const FString& InPath)
 
 void AddSystemPath(const FString& InPath)
 {
+	if (!IFileManager::Get().DirectoryExists(*InPath))
+	{
+		return;
+	}
+
 	if (PyObject* PyPathList = PySys_GetObject(PyCStrCast("path")))
 	{
 		FPyObjectPtr PyPath;
