Fixed MakeTransform() function to take a default value for Scale parameter -> "static FTransform MakeTransform(FVector Location, FRotator Rotation, FVector Scale = FVector(1,1,1))"
  - Removed the meta=(MakeStructureDefaultValue= "1,1,1") from the FTransform struct because Unreal Header Tool is already capable to deduce that from the MakeTransform() default C++ prameter.
  - Removed the meta=(Scale = "1,1,1") from MakeTransform UFUNCTION because that can be deduced by [REDACTED_PERSON] from the default C++ parameter.

Fixed Python code not invoking the a 'Make' function (declared with meta HasNativeMake on a USTRUCT) if no parameters were supplied to Python (ex. unreal.Transform() would not call the MakeTransform(...)).
  - In most case, not calling the Make function worked as expected, but the Make function could implement side effects, so always honor it.

Resaved BP macro assets using MakeTransform() with the default values. The default values saved as string with the asset changed from "1,1,1" to "1.000000,1.000000,1.000000" and this triggered a BP Compiler warning.

The following tests were performed:
  - In the Blueprint Editor, added a Make Transform node and ensured the default values for scale were (1, 1, 1)
  - In the Blueprint Editor, created a User Defined Struct, added as a member a Transform, ensured the default values for the transform scale were (1, 1, 1)
  - Added a Python regression unit tests to ensure unreal.Transform() takes the default scale value as (1, 1, 1) if the scale parameter is not provided.
  - Created a Blueprint Macro Libraries with a Macro calling MakeTransform() with default parameters, saved it, applied the fix, got the warning, refreshed the node, resaved the macro and ensured the macro used the correct default scale.

diff --git a/Engine/Plugins/Experimental/PythonScriptPlugin/Source/PythonScriptPlugin/Private/PyWrapperStruct.cpp b/Engine/Plugins/Experimental/PythonScriptPlugin/Source/PythonScriptPlugin/Private/PyWrapperStruct.cpp
--- a/Engine/Plugins/Experimental/PythonScriptPlugin/Source/PythonScriptPlugin/Private/PyWrapperStruct.cpp
+++ b/Engine/Plugins/Experimental/PythonScriptPlugin/Source/PythonScriptPlugin/Private/PyWrapperStruct.cpp
@@ -351,16 +351,16 @@ int FPyWrapperStruct::MakeStruct(FPyWrapperStruct* InSelf, PyObject* InArgs, PyO
 		return -1;
 	}
 
-	// We can early out if we have no data to apply
-	if (PyTuple_Size(InArgs) == 0 && (!InKwds || PyDict_Size(InKwds) == 0))
+	// If this struct has a custom make function, use that rather than the generic version
+	if (StructMetaData->MakeFunc.Func)
 	{
-		return 0;
+		return CallMakeFunction_Impl(InSelf, InArgs, InKwds, StructMetaData->MakeFunc);
 	}
 
-	// If this struct has a custom make function, use that rather than use the generic version
-	if (StructMetaData->MakeFunc.Func)
+	// We can early out if we have no data to apply
+	if (PyTuple_Size(InArgs) == 0 && (!InKwds || PyDict_Size(InKwds) == 0))
 	{
-		return CallMakeFunction_Impl(InSelf, InArgs, InKwds, StructMetaData->MakeFunc);
+		return 0;
 	}
 
 	// Generic implementation just tries to assign each property
diff --git a/Engine/Source/Runtime/CoreUObject/Public/UObject/NoExportTypes.h b/Engine/Source/Runtime/CoreUObject/Public/UObject/NoExportTypes.h
--- a/Engine/Source/Runtime/CoreUObject/Public/UObject/NoExportTypes.h
+++ b/Engine/Source/Runtime/CoreUObject/Public/UObject/NoExportTypes.h
@@ -1516,7 +1516,7 @@ struct FTransform
 	FVector Translation;
 
 	/** 3D scale (always applied in local space) as a vector. */
-	UPROPERTY(EditAnywhere, BlueprintReadWrite, Category=Transform, SaveGame, meta=(MakeStructureDefaultValue = "1,1,1"))
+	UPROPERTY(EditAnywhere, BlueprintReadWrite, Category=Transform, SaveGame)
 	FVector Scale3D;
 };
 
diff --git a/Engine/Source/Runtime/Engine/Classes/Kismet/KismetMathLibrary.h b/Engine/Source/Runtime/Engine/Classes/Kismet/KismetMathLibrary.h
--- a/Engine/Source/Runtime/Engine/Classes/Kismet/KismetMathLibrary.h
+++ b/Engine/Source/Runtime/Engine/Classes/Kismet/KismetMathLibrary.h
@@ -3722,8 +3722,8 @@ class ENGINE_API UKismetMathLibrary : public UBlueprintFunctionLibrary
 	//
 	
 	/** Make a transform from location, rotation and scale */
-	UFUNCTION(BlueprintPure, meta = (Scale = "1,1,1", Keywords = "construct build", NativeMakeFunc), Category = "Math|Transform")
-	static FTransform MakeTransform(FVector Location, FRotator Rotation, FVector Scale);
+	UFUNCTION(BlueprintPure, meta = (Keywords = "construct build", NativeMakeFunc), Category = "Math|Transform")
+	static FTransform MakeTransform(FVector Location, FRotator Rotation, FVector Scale = FVector(1,1,1));
 
 	/** Breaks apart a transform into location, rotation and scale */
 	UFUNCTION(BlueprintPure, Category = "Math|Transform", meta = (NativeBreakFunc))
