Fixed Python Plugin description mentioning 'UE4' rather than 'Unreal Editor' in the text displayed in the plugin selection page.

diff --git a/Engine/Plugins/Experimental/PythonScriptPlugin/PythonScriptPlugin.uplugin b/Engine/Plugins/Experimental/PythonScriptPlugin/PythonScriptPlugin.uplugin
--- a/Engine/Plugins/Experimental/PythonScriptPlugin/PythonScriptPlugin.uplugin
+++ b/Engine/Plugins/Experimental/PythonScriptPlugin/PythonScriptPlugin.uplugin
@@ -3,7 +3,7 @@
 	"Version": 1,
 	"VersionName": "1.0",
 	"FriendlyName": "Python Editor Script Plugin",
-	"Description": "Python integration for the UE4 editor.",
+	"Description": "Python integration for the Unreal Editor.",
 	"Category": "Scripting",
 	"CreatedBy": "Epic Games, Inc.",
 	"CreatedByURL": "http://epicgames.com",
