Instrumented Python plugin loading time and saved about 500ms to the total module startup time.

diff --git a/Engine/Plugins/Experimental/PythonScriptPlugin/Source/PythonScriptPlugin/Private/PyUtil.cpp b/Engine/Plugins/Experimental/PythonScriptPlugin/Source/PythonScriptPlugin/Private/PyUtil.cpp
--- a/Engine/Plugins/Experimental/PythonScriptPlugin/Source/PythonScriptPlugin/Private/PyUtil.cpp
+++ b/Engine/Plugins/Experimental/PythonScriptPlugin/Source/PythonScriptPlugin/Private/PyUtil.cpp
@@ -22,6 +22,7 @@
 #include "Misc/ScopeExit.h"
 #include "Misc/MessageDialog.h"
 #include "Misc/DefaultValueHelper.h"
+#include "ProfilingDebugging/CpuProfilerTrace.h"
 #include "HAL/FileManager.h"
 #include "UObject/UnrealType.h"
 #include "UObject/EnumProperty.h"
@@ -944,6 +945,8 @@ FOnDiskModules& GetOnDiskUnrealModulesCache()
 
 bool IsModuleAvailableForImport(const TCHAR* InModuleName, const FOnDiskModules* InOnDiskModules, FString* OutResolvedFile)
 {
+	TRACE_CPUPROFILER_EVENT_SCOPE(PyUtil::IsModuleAvailableForImport)
+
 	// Check the sys.modules table first since it avoids hitting the filesystem
 	if (PyObject* PyModulesDict = PySys_GetObject(PyCStrCast("modules")))
 	{
@@ -1055,6 +1058,11 @@ FString GetInterpreterExecutablePath(bool* OutIsEnginePython)
 
 void AddSitePackagesPath(const FString& InPath)
 {
+	if (!IFileManager::Get().DirectoryExists(*InPath))
+	{
+		return;
+	}
+
 	if (FPyObjectPtr PySiteModule = FPyObjectPtr::StealReference(PyImport_ImportModule("site")))
 	{
 		PyObject* PySiteDict = PyModule_GetDict(PySiteModule);
diff --git a/Engine/Plugins/Experimental/PythonScriptPlugin/Source/PythonScriptPlugin/Private/PyWrapperTypeRegistry.cpp b/Engine/Plugins/Experimental/PythonScriptPlugin/Source/PythonScriptPlugin/Private/PyWrapperTypeRegistry.cpp
--- a/Engine/Plugins/Experimental/PythonScriptPlugin/Source/PythonScriptPlugin/Private/PyWrapperTypeRegistry.cpp
+++ b/Engine/Plugins/Experimental/PythonScriptPlugin/Source/PythonScriptPlugin/Private/PyWrapperTypeRegistry.cpp
@@ -18,6 +18,7 @@
 #include "PyFileWriter.h"
 #include "PythonScriptPluginSettings.h"
 #include "ProfilingDebugging/ScopedTimers.h"
+#include "ProfilingDebugging/CpuProfilerTrace.h"
 #include "Misc/Paths.h"
 #include "Misc/FileHelper.h"
 #include "SourceCodeNavigation.h"
@@ -449,6 +450,8 @@ const IPyWrapperInlineStructFactory* FPyWrapperTypeRegistry::GetInlineStructFact
 
 void FPyWrapperTypeRegistry::GenerateWrappedTypes()
 {
+	TRACE_CPUPROFILER_EVENT_SCOPE(FPyWrapperTypeRegistry::GenerateWrappedTypes)
+
 	FGeneratedWrappedTypeReferences GeneratedWrappedTypeReferences;
 	TSet<FName> DirtyModules;
 
@@ -476,6 +479,8 @@ void FPyWrapperTypeRegistry::GenerateWrappedTypes()
 
 void FPyWrapperTypeRegistry::GenerateWrappedTypesForModule(const FName ModuleName)
 {
+	TRACE_CPUPROFILER_EVENT_SCOPE(FPyWrapperTypeRegistry::GenerateWrappedTypesForModule)
+
 	UPackage* const ModulePackage = FindPackage(nullptr, *(FString("/Script/") + ModuleName.ToString()));
 	if (!ModulePackage)
 	{
@@ -509,6 +514,8 @@ void FPyWrapperTypeRegistry::GenerateWrappedTypesForModule(const FName ModuleNam
 
 void FPyWrapperTypeRegistry::OrphanWrappedTypesForModule(const FName ModuleName)
 {
+	TRACE_CPUPROFILER_EVENT_SCOPE(FPyWrapperTypeRegistry::OrphanWrappedTypesForModule)
+
 	TArray<FName> ModuleTypeNames;
 	GeneratedWrappedTypesForModule.MultiFind(ModuleName, ModuleTypeNames, true);
 	GeneratedWrappedTypesForModule.Remove(ModuleName);
@@ -531,6 +538,8 @@ void FPyWrapperTypeRegistry::OrphanWrappedTypesForModule(const FName ModuleName)
 
 void FPyWrapperTypeRegistry::GenerateWrappedTypesForReferences(const FGeneratedWrappedTypeReferences& InGeneratedWrappedTypeReferences, TSet<FName>& OutDirtyModules)
 {
+	TRACE_CPUPROFILER_EVENT_SCOPE(FPyWrapperTypeRegistry::GenerateWrappedTypesForReferences)
+
 	static const EPyTypeGenerationFlags ReferenceGenerationFlags = EPyTypeGenerationFlags::ForceShouldExport | EPyTypeGenerationFlags::IncludeBlueprintGeneratedTypes;
 
 	if (!InGeneratedWrappedTypeReferences.HasReferences())
@@ -566,6 +575,8 @@ void FPyWrapperTypeRegistry::GenerateWrappedTypesForReferences(const FGeneratedW
 
 void FPyWrapperTypeRegistry::NotifyModulesDirtied(const TSet<FName>& InDirtyModules) const
 {
+	TRACE_CPUPROFILER_EVENT_SCOPE(FPyWrapperTypeRegistry::NotifyModulesDirtied)
+
 	for (const FName& DirtyModule : InDirtyModules)
 	{
 		const FString PythonModuleName = PyGenUtil::GetModulePythonName(DirtyModule, false);
diff --git a/Engine/Plugins/Experimental/PythonScriptPlugin/Source/PythonScriptPlugin/Private/PythonScriptPlugin.cpp b/Engine/Plugins/Experimental/PythonScriptPlugin/Source/PythonScriptPlugin/Private/PythonScriptPlugin.cpp
--- a/Engine/Plugins/Experimental/PythonScriptPlugin/Source/PythonScriptPlugin/Private/PythonScriptPlugin.cpp
+++ b/Engine/Plugins/Experimental/PythonScriptPlugin/Source/PythonScriptPlugin/Private/PythonScriptPlugin.cpp
@@ -26,11 +26,13 @@
 #include "Misc/Paths.h"
 #include "Misc/FileHelper.h"
 #include "Misc/PackageName.h"
+#include "HAL/FileManager.h"
 #include "HAL/PlatformMisc.h"
 #include "HAL/PlatformProcess.h"
 #include "Containers/Ticker.h"
 #include "Features/IModularFeatures.h"
 #include "ProfilingDebugging/ScopedTimers.h"
+#include "ProfilingDebugging/CpuProfilerTrace.h"
 #include "Stats/Stats.h"
 #include "String/Find.h"
 
@@ -698,6 +700,8 @@ void FPythonScriptPlugin::PostChange(const UUserDefinedEnum* Enum, FEnumEditorUt
 
 void FPythonScriptPlugin::InitializePython()
 {
+	TRACE_CPUPROFILER_EVENT_SCOPE(FPythonScriptPlugin::InitializePython)
+
 	bInitialized = true;
 
 	const UPythonScriptPluginSettings* PythonPluginSettings = GetDefault<UPythonScriptPluginSettings>();
@@ -1119,6 +1123,8 @@ void FPythonScriptPlugin::RequestStubCodeGeneration()
 
 void FPythonScriptPlugin::GenerateStubCode()
 {
+	TRACE_CPUPROFILER_EVENT_SCOPE(FPythonScriptPlugin::GenerateStubCode)
+
 	if (IsDeveloperModeEnabled())
 	{
 		// Generate stub code if developer mode enabled
@@ -1184,6 +1190,8 @@ void FPythonScriptPlugin::SyncRemoteExecutionToSettings()
 
 void FPythonScriptPlugin::ImportUnrealModule(const TCHAR* InModuleName)
 {
+	TRACE_CPUPROFILER_EVENT_SCOPE_TEXT(*FString::Printf(TEXT("FPythonScriptPlugin::ImportUnrealModule %s"), InModuleName));
+
 	const FString PythonModuleName = FString::Printf(TEXT("unreal_%s"), InModuleName);
 	const FString NativeModuleName = FString::Printf(TEXT("_unreal_%s"), InModuleName);
 
@@ -1252,6 +1260,8 @@ PyObject* FPythonScriptPlugin::EvalString(const TCHAR* InStr, const TCHAR* InCon
 
 PyObject* FPythonScriptPlugin::EvalString(const TCHAR* InStr, const TCHAR* InContext, const int InMode, PyObject* InGlobalDict, PyObject* InLocalDict)
 {
+	TRACE_CPUPROFILER_EVENT_SCOPE(FPythonScriptPlugin::EvalString)
+
 	PyCompilerFlags *PyCompFlags = nullptr;
 
 	PyArena* PyArena = PyArena_New();
@@ -1279,6 +1289,8 @@ PyObject* FPythonScriptPlugin::EvalString(const TCHAR* InStr, const TCHAR* InCon
 
 bool FPythonScriptPlugin::RunString(FPythonCommandEx& InOutPythonCommand)
 {
+	TRACE_CPUPROFILER_EVENT_SCOPE(FPythonScriptPlugin::RunString)
+
 	// Execute Python code within this block
 	{
 		FPyScopedGIL GIL;
@@ -1321,6 +1333,8 @@ bool FPythonScriptPlugin::RunString(FPythonCommandEx& InOutPythonCommand)
 
 bool FPythonScriptPlugin::RunFile(const TCHAR* InFile, const TCHAR* InArgs, FPythonCommandEx& InOutPythonCommand)
 {
+	TRACE_CPUPROFILER_EVENT_SCOPE(FPythonScriptPlugin::RunFile)
+
 	auto ResolveFilePath = [InFile]() -> FString
 	{
 		// Favor the CWD
@@ -1429,6 +1443,8 @@ void FPythonScriptPlugin::OnModuleDirtied(FName InModuleName)
 
 void FPythonScriptPlugin::OnModulesChanged(FName InModuleName, EModuleChangeReason InModuleChangeReason)
 {
+	TRACE_CPUPROFILER_EVENT_SCOPE(FPythonScriptPlugin::OnModulesChanged)
+
 	switch (InModuleChangeReason)
 	{
 	case EModuleChangeReason::ModuleLoaded:
@@ -1454,6 +1470,8 @@ void FPythonScriptPlugin::OnModulesChanged(FName InModuleName, EModuleChangeReas
 
 void FPythonScriptPlugin::OnContentPathMounted(const FString& InAssetPath, const FString& InFilesystemPath)
 {
+	TRACE_CPUPROFILER_EVENT_SCOPE(FPythonScriptPlugin::OnContentPathMounted)
+
 	{
 		FPyScopedGIL GIL;
 		RegisterModulePaths(InFilesystemPath);
@@ -1469,6 +1487,8 @@ void FPythonScriptPlugin::OnContentPathMounted(const FString& InAssetPath, const
 
 void FPythonScriptPlugin::OnContentPathDismounted(const FString& InAssetPath, const FString& InFilesystemPath)
 {
+	TRACE_CPUPROFILER_EVENT_SCOPE(FPythonScriptPlugin::OnContentPathDismounted)
+
 	{
 		FPyScopedGIL GIL;
 		UnregisterModulePaths(InFilesystemPath);
@@ -1484,19 +1504,26 @@ void FPythonScriptPlugin::OnContentPathDismounted(const FString& InAssetPath, co
 
 void FPythonScriptPlugin::RegisterModulePaths(const FString& InFilesystemPath)
 {
+	TRACE_CPUPROFILER_EVENT_SCOPE(FPythonScriptPlugin::RegisterModulePaths)
+
 	const FString PythonContentPath = FPaths::ConvertRelativePathToFull(InFilesystemPath / TEXT("Python"));
-	PyUtil::AddSystemPath(PythonContentPath);
+	if (IFileManager::Get().DirectoryExists(*PythonContentPath))
+	{
+		PyUtil::AddSystemPath(PythonContentPath);
 
-	const FString PythonContentPlatformSitePackagesPath = PythonContentPath / TEXT("Lib") / FPlatformMisc::GetUBTPlatform() / TEXT("site-packages");
-	const FString PythonContentGeneralSitePackagesPath = PythonContentPath / TEXT("Lib") / TEXT("site-packages");
-	PyUtil::AddSitePackagesPath(PythonContentPlatformSitePackagesPath);
-	PyUtil::AddSitePackagesPath(PythonContentGeneralSitePackagesPath);
+		const FString PythonContentPlatformSitePackagesPath = PythonContentPath / TEXT("Lib") / FPlatformMisc::GetUBTPlatform() / TEXT("site-packages");
+		const FString PythonContentGeneralSitePackagesPath = PythonContentPath / TEXT("Lib") / TEXT("site-packages");
+		PyUtil::AddSitePackagesPath(PythonContentPlatformSitePackagesPath);
+		PyUtil::AddSitePackagesPath(PythonContentGeneralSitePackagesPath);
 
-	PyUtil::GetOnDiskUnrealModulesCache().AddModules(*PythonContentPath);
+		PyUtil::GetOnDiskUnrealModulesCache().AddModules(*PythonContentPath);
+	}
 }
 
 void FPythonScriptPlugin::UnregisterModulePaths(const FString& InFilesystemPath)
 {
+	TRACE_CPUPROFILER_EVENT_SCOPE(FPythonScriptPlugin::UnregisterModulePaths)
+
 	const FString PythonContentPath = FPaths::ConvertRelativePathToFull(InFilesystemPath / TEXT("Python"));
 	PyUtil::RemoveSystemPath(PythonContentPath);
 
@@ -1515,6 +1542,8 @@ bool FPythonScriptPlugin::IsDeveloperModeEnabled()
 
 void FPythonScriptPlugin::OnAssetRenamed(const FAssetData& Data, const FString& OldName)
 {
+	TRACE_CPUPROFILER_EVENT_SCOPE(FPythonScriptPlugin::OnAssetRenamed)
+
 	FPyWrapperTypeRegistry& PyWrapperTypeRegistry = FPyWrapperTypeRegistry::Get();
 	const FName OldPackageName = *FPackageName::ObjectPathToPackageName(OldName);
 	
@@ -1535,6 +1564,8 @@ void FPythonScriptPlugin::OnAssetRenamed(const FAssetData& Data, const FString&
 
 void FPythonScriptPlugin::OnAssetRemoved(const FAssetData& Data)
 {
+	TRACE_CPUPROFILER_EVENT_SCOPE(FPythonScriptPlugin::OnAssetRemoved)
+
 	FPyWrapperTypeRegistry& PyWrapperTypeRegistry = FPyWrapperTypeRegistry::Get();
 	
 	// If this asset has an associated Python type, then we need to remove it
@@ -1546,6 +1577,8 @@ void FPythonScriptPlugin::OnAssetRemoved(const FAssetData& Data)
 
 void FPythonScriptPlugin::OnAssetReload(const EPackageReloadPhase InPackageReloadPhase, FPackageReloadedEvent* InPackageReloadedEvent)
 {
+	TRACE_CPUPROFILER_EVENT_SCOPE(FPythonScriptPlugin::OnAssetReload)
+
 	if (InPackageReloadPhase == EPackageReloadPhase::PostPackageFixup)
 	{
 		// Get the primary asset in this package
@@ -1558,6 +1591,8 @@ void FPythonScriptPlugin::OnAssetReload(const EPackageReloadPhase InPackageReloa
 
 void FPythonScriptPlugin::OnAssetUpdated(const UObject* InObj)
 {
+	TRACE_CPUPROFILER_EVENT_SCOPE(FPythonScriptPlugin::OnAssetUpdated)
+
 	if (const UObject* AssetPtr = PyGenUtil::GetAssetTypeRegistryType(InObj))
 	{
 		// If this asset has an associated Python type, then we need to re-generate it
