Fixed Python GIL (Global Interpreter Lock) deadlock while loading an asset.
  - Released the GIL while loading or finding an asset.
  - Released the GIL while loading or finding an object.
  - Released the GIL while loading or finding packages.
  - Released the GIL while logging/flushing stdout/stderr from Python.

diff --git a/Engine/Plugins/Experimental/PythonScriptPlugin/Source/PythonScriptPlugin/Private/PyCore.cpp b/Engine/Plugins/Experimental/PythonScriptPlugin/Source/PythonScriptPlugin/Private/PyCore.cpp
--- a/Engine/Plugins/Experimental/PythonScriptPlugin/Source/PythonScriptPlugin/Private/PyCore.cpp
+++ b/Engine/Plugins/Experimental/PythonScriptPlugin/Source/PythonScriptPlugin/Private/PyCore.cpp
@@ -1035,7 +1035,9 @@ PyObject* Log(PyObject* InSelf, PyObject* InArgs)
 		const FString LogMessage = PyUtil::PyObjectToUEString(PyObj);
 		if (ShouldLog(LogMessage))
 		{
+			Py_BEGIN_ALLOW_THREADS
 			UE_LOG(LogPython, Log, TEXT("%s"), *LogMessage);
+			Py_END_ALLOW_THREADS
 			GetPythonLogCapture().Broadcast(EPythonLogOutputType::Info, *LogMessage);
 		}
 
@@ -1053,7 +1055,9 @@ PyObject* LogWarning(PyObject* InSelf, PyObject* InArgs)
 		const FString LogMessage = PyUtil::PyObjectToUEString(PyObj);
 		if (ShouldLog(LogMessage))
 		{
+			Py_BEGIN_ALLOW_THREADS
 			UE_LOG(LogPython, Warning, TEXT("%s"), *LogMessage);
+			Py_END_ALLOW_THREADS
 			GetPythonLogCapture().Broadcast(EPythonLogOutputType::Warning, *LogMessage);
 		}
 
@@ -1071,7 +1075,9 @@ PyObject* LogError(PyObject* InSelf, PyObject* InArgs)
 		const FString LogMessage = PyUtil::PyObjectToUEString(PyObj);
 		if (ShouldLog(LogMessage))
 		{
+			Py_BEGIN_ALLOW_THREADS
 			UE_LOG(LogPython, Error, TEXT("%s"), *LogMessage);
+			Py_END_ALLOW_THREADS
 			GetPythonLogCapture().Broadcast(EPythonLogOutputType::Error, *LogMessage);
 		}
 
@@ -1085,7 +1091,9 @@ PyObject* LogFlush(PyObject* InSelf)
 {
 	if (GLog)
 	{
+		Py_BEGIN_ALLOW_THREADS
 		GLog->Flush();
+		Py_END_ALLOW_THREADS
 	}
 
 	Py_RETURN_NONE;
@@ -1272,7 +1280,11 @@ PyObject* FindObject(PyObject* InSelf, PyObject* InArgs, PyObject* InKwds)
 {
 	return FindOrLoadObjectImpl(InSelf, InArgs, InKwds, TEXT("find_object"), [](UClass* ObjectType, UObject* ObjectOuter, const TCHAR* ObjectName)
 	{
-		return ::StaticFindObject(ObjectType, ObjectOuter, ObjectName);
+		UObject* Object = nullptr;
+		Py_BEGIN_ALLOW_THREADS
+		Object = ::StaticFindObject(ObjectType, ObjectOuter, ObjectName);
+		Py_END_ALLOW_THREADS
+		return Object;
 	});
 }
 
@@ -1280,7 +1292,11 @@ PyObject* LoadObject(PyObject* InSelf, PyObject* InArgs, PyObject* InKwds)
 {
 	return FindOrLoadObjectImpl(InSelf, InArgs, InKwds, TEXT("load_object"), [](UClass* ObjectType, UObject* ObjectOuter, const TCHAR* ObjectName)
 	{
-		return ::StaticLoadObject(ObjectType, ObjectOuter, ObjectName);
+		UObject* Object = nullptr;
+		Py_BEGIN_ALLOW_THREADS
+		Object = ::StaticLoadObject(ObjectType, ObjectOuter, ObjectName);
+		Py_END_ALLOW_THREADS
+		return Object;
 	});
 }
 
@@ -1392,7 +1408,11 @@ PyObject* FindAsset(PyObject* InSelf, PyObject* InArgs, PyObject* InKwds)
 {
 	return FindOrLoadAssetImpl(InSelf, InArgs, InKwds, TEXT("find_asset"), [](UClass* ObjectType, UObject* ObjectOuter, const TCHAR* ObjectName)
 	{
-		return ::StaticFindObject(ObjectType, ObjectOuter, ObjectName);
+		UObject* Object = nullptr;
+		Py_BEGIN_ALLOW_THREADS
+		Object = ::StaticFindObject(ObjectType, ObjectOuter, ObjectName);
+		Py_END_ALLOW_THREADS
+		return Object;
 	});
 }
 
@@ -1400,7 +1420,11 @@ PyObject* LoadAsset(PyObject* InSelf, PyObject* InArgs, PyObject* InKwds)
 {
 	return FindOrLoadAssetImpl(InSelf, InArgs, InKwds, TEXT("load_asset"), [](UClass* ObjectType, UObject* ObjectOuter, const TCHAR* ObjectName)
 	{
-		return ::StaticLoadObject(ObjectType, ObjectOuter, ObjectName);
+		UObject* Object = nullptr;
+		Py_BEGIN_ALLOW_THREADS
+		Object = ::StaticLoadObject(ObjectType, ObjectOuter, ObjectName);
+		Py_END_ALLOW_THREADS
+		return Object;
 	});
 }
 
@@ -1427,7 +1451,11 @@ PyObject* FindPackage(PyObject* InSelf, PyObject* InArgs)
 {
 	return FindOrLoadPackageImpl(InSelf, InArgs, TEXT("find_package"), [](const TCHAR* PackageName)
 	{
-		return ::FindPackage(nullptr, PackageName);
+		UPackage* Package = nullptr;
+		Py_BEGIN_ALLOW_THREADS
+		Package = ::FindPackage(nullptr, PackageName);
+		Py_END_ALLOW_THREADS
+		return Package;
 	});
 }
 
@@ -1435,7 +1463,11 @@ PyObject* LoadPackage(PyObject* InSelf, PyObject* InArgs)
 {
 	return FindOrLoadPackageImpl(InSelf, InArgs, TEXT("load_package"), [](const TCHAR* PackageName)
 	{
-		return ::LoadPackage(nullptr, PackageName, LOAD_None);
+		UPackage* Package = nullptr;
+		Py_BEGIN_ALLOW_THREADS
+		Package = ::LoadPackage(nullptr, PackageName, LOAD_None);
+		Py_END_ALLOW_THREADS
+		return Package;
 	});
 }
 
