Fixed Python plugin to allow resolving USTRUCT Make/Break functions specifed with 'HasNativeMake' and 'HasNativeBreak' later. This supports cases where the function are implemented in a C++ module that is loaded after the module declaring the USTRUCT. The implementation postpones resolving missing functions until FCoreDelegates::OnPostEngineInit is invoked. If one or more functions are still missing, corresponding error(s) are logged.

diff --git a/Engine/Plugins/Experimental/PythonScriptPlugin/Source/PythonScriptPlugin/Private/PyWrapperTypeRegistry.cpp b/Engine/Plugins/Experimental/PythonScriptPlugin/Source/PythonScriptPlugin/Private/PyWrapperTypeRegistry.cpp
--- a/Engine/Plugins/Experimental/PythonScriptPlugin/Source/PythonScriptPlugin/Private/PyWrapperTypeRegistry.cpp
+++ b/Engine/Plugins/Experimental/PythonScriptPlugin/Source/PythonScriptPlugin/Private/PyWrapperTypeRegistry.cpp
@@ -34,6 +34,90 @@
 
 #if WITH_PYTHON
 
+namespace UE
+{
+namespace Python
+{
+
+/**
+ * Finds the UFunction corresponding to the name specified by 'HasNativeMake' or 'HasNativeBreak' meta data key.
+ * @param The structure to inspect for the 'HasNativeMake' or 'HasNativeBreak' meta data keys.
+ * @param InMetaDataKey The meta data key name. Can only be 'HasNativeMake' or 'HasNativeBreak'.
+ * @param NotFoundFn Function invoked if the structure specifies as Make or Break function, but the function couldn't be found.
+ * @return The function, if the struct has the meta key and if the function was found. Null otherwise.
+ */
+template<typename NotFoundFuncT>
+UFunction* FindMakeBreakFunction(const UScriptStruct* InStruct, const FName& InMetaDataKey, const NotFoundFuncT& NotFoundFn)
+{
+	check(InMetaDataKey == PyGenUtil::HasNativeMakeMetaDataKey || InMetaDataKey == PyGenUtil::HasNativeBreakMetaDataKey);
+
+	UFunction* MakeBreakFunc = nullptr;
+
+	const FString MakeBreakFunctionName = InStruct->GetMetaData(InMetaDataKey);
+	if (!MakeBreakFunctionName.IsEmpty())
+	{
+		// Find the function.
+		MakeBreakFunc = FindObject<UFunction>(/*Outer*/nullptr, *MakeBreakFunctionName, /*ExactClass*/true);
+		if (!MakeBreakFunc)
+		{
+			NotFoundFn(MakeBreakFunctionName);
+		}
+	}
+	return MakeBreakFunc;
+}
+
+void SetMakeFunction(FPyWrapperStructMetaData& MetaData, UFunction* MakeFunc)
+{
+	check(MetaData.Struct != nullptr && MakeFunc != nullptr);
+
+	MetaData.MakeFunc.SetFunction(MakeFunc);
+
+	const bool bHasValidReturn =
+		MetaData.MakeFunc.OutputParams.Num() == 1 &&
+		MetaData.MakeFunc.OutputParams[0].ParamProp->IsA<FStructProperty>() &&
+		CastFieldChecked<const FStructProperty>(MetaData.MakeFunc.OutputParams[0].ParamProp)->Struct == MetaData.Struct;
+
+	if (!bHasValidReturn)
+	{
+		REPORT_PYTHON_GENERATION_ISSUE(Error, TEXT("Struct '%s' is marked as 'HasNativeMake' but the function '%s' does not return the struct type."), *MetaData.Struct->GetName(), *MetaData.MakeFunc.Func->GetPathName());
+		MetaData.MakeFunc.SetFunction(nullptr);
+		return;
+	}
+
+	// Set the make arguments to be optional to mirror the behavior of struct InitParams
+	for (PyGenUtil::FGeneratedWrappedMethodParameter& InputParam : MetaData.MakeFunc.InputParams)
+	{
+		if (!InputParam.ParamDefaultValue.IsSet())
+		{
+			InputParam.ParamDefaultValue = FString();
+		}
+	}
+}
+
+void SetBreakFunction(FPyWrapperStructMetaData& MetaData, UFunction* BreakFunc)
+{
+	check(MetaData.Struct != nullptr && BreakFunc != nullptr);
+
+	MetaData.BreakFunc.SetFunction(BreakFunc);
+
+	const bool bHasValidInput =
+		MetaData.BreakFunc.InputParams.Num() == 1 &&
+		MetaData.BreakFunc.InputParams[0].ParamProp->IsA<FStructProperty>() &&
+		CastFieldChecked<const FStructProperty>(MetaData.BreakFunc.InputParams[0].ParamProp)->Struct == MetaData.Struct;
+
+	if (!bHasValidInput)
+	{
+		REPORT_PYTHON_GENERATION_ISSUE(Error, TEXT("Struct '%s' is marked as 'HasNativeBreak' but the function '%s' does not have the struct type as its only input argument."), *MetaData.Struct->GetName(), *MetaData.BreakFunc.Func->GetPathName());
+		MetaData.BreakFunc.SetFunction(nullptr);
+	}
+};
+
+
+} // namespace Python
+} // namespace UE
+
+
+
 DECLARE_FLOAT_ACCUMULATOR_STAT(TEXT("Generate Wrapped Class Total Time"), STAT_GenerateWrappedClassTotalTime, STATGROUP_Python);
 DECLARE_DWORD_ACCUMULATOR_STAT(TEXT("Generate Wrapped Class Call Count"), STAT_GenerateWrappedClassCallCount, STATGROUP_Python);
 DECLARE_DWORD_ACCUMULATOR_STAT(TEXT("Generate Wrapped Class Obj Count"), STAT_GenerateWrappedClassObjCount, STATGROUP_Python);
@@ -424,6 +508,18 @@ void FPyWrapperTypeReinstancer::AddReferencedObjects(FReferenceCollector& InColl
 FPyWrapperTypeRegistry::FPyWrapperTypeRegistry()
 	: bCanRegisterInlineStructFactories(true)
 {
+	// When everything is loaded, warn the user if there are still unresolved Make/Break functions.
+	FCoreDelegates::OnPostEngineInit.AddLambda([this]()
+	{
+		for (const TPair<FString, TSharedPtr<FPyWrapperStructMetaData>>& UnresolvedMakePair : UnresolvedMakeFuncs)
+		{
+			REPORT_PYTHON_GENERATION_ISSUE(Error, TEXT("Struct '%s' is marked as '%s' but the function '%s' could not be found."), *UnresolvedMakePair.Value->Struct->GetName(), *PyGenUtil::HasNativeMakeMetaDataKey.ToString(), *UnresolvedMakePair.Key);
+		}
+		for (const TPair<FString, TSharedPtr<FPyWrapperStructMetaData>>& UnresolvedBreakPair : UnresolvedBreakFuncs)
+		{
+			REPORT_PYTHON_GENERATION_ISSUE(Error, TEXT("Struct '%s' is marked as '%s' but the function '%s' could not be found."), *UnresolvedBreakPair.Value->Struct->GetName(), *PyGenUtil::HasNativeBreakMetaDataKey.ToString(), *UnresolvedBreakPair.Key);
+		}
+	});
 }
 
 FPyWrapperTypeRegistry& FPyWrapperTypeRegistry::Get()
@@ -507,6 +603,25 @@ void FPyWrapperTypeRegistry::GenerateWrappedTypesForModule(const FName ModuleNam
 		GenerateWrappedTypesForReferences(GeneratedWrappedTypeReferences, DirtyModules);
 	}
 
+	// Try resolving missing make/break native functions. Some USTRUCT uses make/break functions implemented in other C++ module(s)
+	// that might not be loaded yet when Python glue was generated for the struct.
+	for (TMap<FString, TSharedPtr<FPyWrapperStructMetaData>>::TIterator It = UnresolvedMakeFuncs.CreateIterator(); It; ++It)
+	{
+		if (UFunction* MakeFunc = FindObject<UFunction>(nullptr, *It->Key, true))
+		{
+			UE::Python::SetMakeFunction(*It->Value, MakeFunc);
+			It.RemoveCurrent();
+		}
+	}
+	for (TMap<FString, TSharedPtr<FPyWrapperStructMetaData>>::TIterator It = UnresolvedBreakFuncs.CreateIterator(); It; ++It)
+	{
+		if (UFunction* BreakFunc = FindObject<UFunction>(nullptr, *It->Key, true))
+		{
+			UE::Python::SetBreakFunction(*It->Value, BreakFunc);
+			It.RemoveCurrent();
+		}
+	}
+
 	NotifyModulesDirtied(DirtyModules);
 
 	UE_LOG(LogPython, Verbose, TEXT("Took %f seconds to generate and initialize Python wrapped types for '%s'."), GenerateDuration, *ModuleName.ToString());
@@ -1577,67 +1692,25 @@ PyTypeObject* FPyWrapperTypeRegistry::GenerateWrappedStructType(const UScriptStr
 	GeneratedWrappedType->PyType.tp_init = (initproc)&FFuncs::Init;
 	GeneratedWrappedType->PyType.tp_flags = Py_TPFLAGS_DEFAULT | Py_TPFLAGS_BASETYPE;
 
-	auto FindMakeBreakFunction = [InStruct](const FName& InKey) -> const UFunction*
-	{
-		const FString MakeBreakName = InStruct->GetMetaData(InKey);
-		if (!MakeBreakName.IsEmpty())
-		{
-			const UFunction* MakeBreakFunc = FindObject<UFunction>(nullptr, *MakeBreakName, true);
-			if (!MakeBreakFunc)
-			{
-				REPORT_PYTHON_GENERATION_ISSUE(Error, TEXT("Struct '%s' is marked as '%s' but the function '%s' could not be found."), *InStruct->GetName(), *InKey.ToString(), *MakeBreakName);
-			}
-			return MakeBreakFunc;
-		}
-		return nullptr;
-	};
+	TSharedRef<FPyWrapperStructMetaData> StructMetaData = MakeShared<FPyWrapperStructMetaData>();
+	StructMetaData->Struct = (UScriptStruct*)InStruct;
+	StructMetaData->PythonProperties = MoveTemp(PythonProperties);
+	StructMetaData->PythonDeprecatedProperties = MoveTemp(PythonDeprecatedProperties);
 
-	auto FindMakeFunction = [InStruct, &FindMakeBreakFunction]() -> PyGenUtil::FGeneratedWrappedFunction
+	// If the struct has the 'HasNativeMake' meta and the function is found, set it, otherwise, postpone maybe its going to be loaded later.
+	if (UFunction* MakeFunc = UE::Python::FindMakeBreakFunction(InStruct, PyGenUtil::HasNativeMakeMetaDataKey,
+		[this, &StructMetaData](const FString& MakeFuncName){ UnresolvedMakeFuncs.Emplace(MakeFuncName, StructMetaData); }))
 	{
-		PyGenUtil::FGeneratedWrappedFunction MakeFunc;
-		MakeFunc.SetFunction(FindMakeBreakFunction(PyGenUtil::HasNativeMakeMetaDataKey));
-		if (MakeFunc.Func)
-		{
-			const bool bHasValidReturn = MakeFunc.OutputParams.Num() == 1 && MakeFunc.OutputParams[0].ParamProp->IsA<FStructProperty>() && CastFieldChecked<const FStructProperty>(MakeFunc.OutputParams[0].ParamProp)->Struct == InStruct;
-			if (!bHasValidReturn)
-			{
-				REPORT_PYTHON_GENERATION_ISSUE(Error, TEXT("Struct '%s' is marked as 'HasNativeMake' but the function '%s' does not return the struct type."), *InStruct->GetName(), *MakeFunc.Func->GetPathName());
-				MakeFunc.SetFunction(nullptr);
-			}
-			// Set the make arguments to be optional to mirror the behavior of struct InitParams
-			for (PyGenUtil::FGeneratedWrappedMethodParameter& InputParam : MakeFunc.InputParams)
-			{
-				if (!InputParam.ParamDefaultValue.IsSet())
-				{
-					InputParam.ParamDefaultValue = FString();
-				}
-			}
-		}
-		return MakeFunc;
-	};
+		UE::Python::SetMakeFunction(*StructMetaData, MakeFunc);
+	}
 
-	auto FindBreakFunction = [InStruct, &FindMakeBreakFunction]() -> PyGenUtil::FGeneratedWrappedFunction
+	// If the struct has the 'HasNativeBreak' meta and the function is found, set it, otherwise, postpone maybe its going to be loaded later.
+	if (UFunction* BreakFunc = UE::Python::FindMakeBreakFunction(InStruct, PyGenUtil::HasNativeBreakMetaDataKey,
+		[this, &StructMetaData](const FString& BreakFuncName) { UnresolvedBreakFuncs.Emplace(BreakFuncName, StructMetaData); }))
 	{
-		PyGenUtil::FGeneratedWrappedFunction BreakFunc;
-		BreakFunc.SetFunction(FindMakeBreakFunction(PyGenUtil::HasNativeBreakMetaDataKey));
-		if (BreakFunc.Func)
-		{
-			const bool bHasValidInput = BreakFunc.InputParams.Num() == 1 && BreakFunc.InputParams[0].ParamProp->IsA<FStructProperty>() && CastFieldChecked<const FStructProperty>(BreakFunc.InputParams[0].ParamProp)->Struct == InStruct;
-			if (!bHasValidInput)
-			{
-				REPORT_PYTHON_GENERATION_ISSUE(Error, TEXT("Struct '%s' is marked as 'HasNativeBreak' but the function '%s' does not have the struct type as its only input argument."), *InStruct->GetName(), *BreakFunc.Func->GetPathName());
-				BreakFunc.SetFunction(nullptr);
-			}
-		}
-		return BreakFunc;
-	};
+		UE::Python::SetBreakFunction(*StructMetaData, BreakFunc);
+	}
 
-	TSharedRef<FPyWrapperStructMetaData> StructMetaData = MakeShared<FPyWrapperStructMetaData>();
-	StructMetaData->Struct = (UScriptStruct*)InStruct;
-	StructMetaData->PythonProperties = MoveTemp(PythonProperties);
-	StructMetaData->PythonDeprecatedProperties = MoveTemp(PythonDeprecatedProperties);
-	StructMetaData->MakeFunc = FindMakeFunction();
-	StructMetaData->BreakFunc = FindBreakFunction();
 	// Build a complete list of init params for this struct (parent struct params + our params)
 	if (SuperPyType)
 	{
diff --git a/Engine/Plugins/Experimental/PythonScriptPlugin/Source/PythonScriptPlugin/Private/PyWrapperTypeRegistry.h b/Engine/Plugins/Experimental/PythonScriptPlugin/Source/PythonScriptPlugin/Private/PyWrapperTypeRegistry.h
--- a/Engine/Plugins/Experimental/PythonScriptPlugin/Source/PythonScriptPlugin/Private/PyWrapperTypeRegistry.h
+++ b/Engine/Plugins/Experimental/PythonScriptPlugin/Source/PythonScriptPlugin/Private/PyWrapperTypeRegistry.h
@@ -25,6 +25,7 @@ struct FPyWrapperArray;
 struct FPyWrapperFixedArray;
 struct FPyWrapperSet;
 struct FPyWrapperMap;
+struct FPyWrapperStructMetaData;
 
 class FPyWrapperOwnerContext;
 class UPythonGeneratedClass;
@@ -505,6 +506,12 @@ private:
 	/** Map from the Unreal module name to its generated type names (names are the Unreal names) */
 	TMultiMap<FName, FName> GeneratedWrappedTypesForModule;
 
+	/** Map make func name to the meta data of the struct to make. Used when meta 'HasNativeMake' references a function not loaded yet. */
+	TMap<FString, TSharedPtr<FPyWrapperStructMetaData>> UnresolvedMakeFuncs;
+
+	/** Map break func name to the meta data of the struct to break. Used when meta 'HasNativeBreak' references a function not loaded yet. */
+	TMap<FString, TSharedPtr<FPyWrapperStructMetaData>> UnresolvedBreakFuncs;
+
 	/** Array of generated Python type data that has been orphaned (due to its owner module being unloaded/reloaded) */
 	TArray<TSharedPtr<PyGenUtil::FGeneratedWrappedType>> OrphanedWrappedTypes;
 
