Fix CIS error in PyWrapperTypeRegistry.cpp.

diff --git a/Engine/Plugins/Experimental/PythonScriptPlugin/Source/PythonScriptPlugin/Private/PyWrapperTypeRegistry.cpp b/Engine/Plugins/Experimental/PythonScriptPlugin/Source/PythonScriptPlugin/Private/PyWrapperTypeRegistry.cpp
--- a/Engine/Plugins/Experimental/PythonScriptPlugin/Source/PythonScriptPlugin/Private/PyWrapperTypeRegistry.cpp
+++ b/Engine/Plugins/Experimental/PythonScriptPlugin/Source/PythonScriptPlugin/Private/PyWrapperTypeRegistry.cpp
@@ -28,6 +28,7 @@
 #include "UObject/Package.h"
 #include "UObject/EnumProperty.h"
 #include "UObject/StructOnScope.h"
+#include "Misc/CoreDelegates.h"
 #if WITH_EDITOR
 #include "Kismet2/ReloadUtilities.h"
 #endif
