Fixed Python GIL (Global Interpreter Lock) deadlock while loading an asset.
  - Released the GIL in Finalize() - rename class, get default object and refresh BlueprintActionDatabase

diff --git a/Engine/Plugins/Experimental/PythonScriptPlugin/Source/PythonScriptPlugin/Private/PyWrapperObject.cpp b/Engine/Plugins/Experimental/PythonScriptPlugin/Source/PythonScriptPlugin/Private/PyWrapperObject.cpp
--- a/Engine/Plugins/Experimental/PythonScriptPlugin/Source/PythonScriptPlugin/Private/PyWrapperObject.cpp
+++ b/Engine/Plugins/Experimental/PythonScriptPlugin/Source/PythonScriptPlugin/Private/PyWrapperObject.cpp
@@ -1271,12 +1271,15 @@ public:
 		{
 			PrepareOldClassForReinstancing();
 		}
+
+		Py_BEGIN_ALLOW_THREADS
 		NewClass->Rename(*ClassName, nullptr, REN_DontCreateRedirectors);
 
 		// Finalize the class
 		NewClass->Bind();
 		NewClass->StaticLink(true);
 		NewClass->AssembleReferenceTokenStream();
+		Py_END_ALLOW_THREADS
 
 		// Add the object meta-data to the type
 		NewClass->PyMetaData.Class = NewClass;
@@ -1287,7 +1290,9 @@ public:
 		FPyWrapperTypeRegistry::Get().RegisterWrappedClassType(NewClass->GetFName() , PyType);
 
 		// Ensure the CDO exists
+		Py_BEGIN_ALLOW_THREADS
 		NewClass->GetDefaultObject();
+		Py_END_ALLOW_THREADS
 
 		// Re-instance the old class and re-parent any derived classes to this new type
 		if (OldClass)
@@ -1298,6 +1303,7 @@ public:
 
 		if (FBlueprintActionDatabase* ActionDB = FBlueprintActionDatabase::TryGet())
 		{
+			Py_BEGIN_ALLOW_THREADS
 			// Notify Blueprints that there is a new class to add to the action list
 			ActionDB->RefreshClassActions(NewClass);
 
@@ -1305,6 +1311,7 @@ public:
 			{
 				ActionDB->RefreshClassActions(OldClass);
 			}
+			Py_END_ALLOW_THREADS
 		}
 
 		// Null the NewClass pointer so the destructor doesn't kill it
