UE-146778 - CIS Content Error: Unhandled implicit casts found during compilation of function

The BP compiler erroneously reports that python nodes might be missing implicit casts. However, UK2Node_ExecutePythonScript nodes handle argument slightly differently than conventional function nodes. The pins don't actually correspond to actual arguments on ExecutePythonScript. Instead, they more or less serve as references to their linked inputs, which ExecutePythonScript will pop from the stack and perform a direct conversion to a Python type. As a result, there's nothing to cast, and we can safely remove the entries from the ImplicitCastMap on the current content.

diff --git a/Engine/Plugins/Experimental/PythonScriptPlugin/Source/PythonScriptPlugin/Private/K2Node_ExecutePythonScript.cpp b/Engine/Plugins/Experimental/PythonScriptPlugin/Source/PythonScriptPlugin/Private/K2Node_ExecutePythonScript.cpp
--- a/Engine/Plugins/Experimental/PythonScriptPlugin/Source/PythonScriptPlugin/Private/K2Node_ExecutePythonScript.cpp
+++ b/Engine/Plugins/Experimental/PythonScriptPlugin/Source/PythonScriptPlugin/Private/K2Node_ExecutePythonScript.cpp
@@ -13,6 +13,8 @@
 #include "Kismet2/BlueprintEditorUtils.h"
 #include "BlueprintNodeSpawner.h"
 #include "BlueprintActionDatabaseRegistrar.h"
+#include "CallFunctionHandler.h"
+#include "KismetCastingUtils.h"
 
 #define LOCTEXT_NAMESPACE "K2Node_ExecutePythonScript"
 
@@ -39,6 +41,37 @@ FString PythonizePinName(const UEdGraphPin* InPin)
 
 }
 
+class FKCHandler_ExecutePythonScript : public FKCHandler_CallFunction
+{
+public:
+	FKCHandler_ExecutePythonScript(FKismetCompilerContext& InCompilerContext)
+		: FKCHandler_CallFunction(InCompilerContext)
+	{
+	}
+
+	virtual void Compile(FKismetFunctionContext& Context, UEdGraphNode* Node) override
+	{
+		using namespace UE::KismetCompiler;
+
+		FKCHandler_CallFunction::Compile(Context, Node);
+
+		UK2Node_ExecutePythonScript* PythonNode = CastChecked<UK2Node_ExecutePythonScript>(Node);
+
+		// The ExecutePythonScript function uses variadic inputs, which obviates the need for implicit casting.
+		// In other words, ExecutePythonScript's native signature doesn't include the parameters corresponding
+		// to the inputs pins. As a result, there's nothing to cast the inputs to. Instead, the implementation of
+		// execExecutePythonScript directly inspects these arguments and converts them to a Python type.
+		for (FName InputName : PythonNode->Inputs)
+		{
+			UEdGraphPin* Pin = PythonNode->FindPinChecked(InputName, EGPD_Input);
+			if (Pin->PinType.PinCategory == UEdGraphSchema_K2::PC_Real)
+			{
+				CastingUtils::RemoveRegisteredImplicitCast(Context, Pin);
+			}
+		}
+	}
+};
+
 UK2Node_ExecutePythonScript::UK2Node_ExecutePythonScript()
 {
 	FunctionReference.SetExternalMember(GET_FUNCTION_NAME_CHECKED(UPythonScriptLibrary, ExecutePythonScript), UPythonScriptLibrary::StaticClass());
@@ -289,6 +322,11 @@ void UK2Node_ExecutePythonScript::ExpandNode(FKismetCompilerContext& CompilerCon
 	Super::ExpandNode(CompilerContext, SourceGraph);
 }
 
+class FNodeHandlingFunctor* UK2Node_ExecutePythonScript::CreateNodeHandler(class FKismetCompilerContext& CompilerContext) const
+{
+	return new FKCHandler_ExecutePythonScript(CompilerContext);
+}
+
 bool UK2Node_ExecutePythonScript::CanPasteHere(const UEdGraph* TargetGraph) const
 {
 	bool bCanPaste = Super::CanPasteHere(TargetGraph);
diff --git a/Engine/Plugins/Experimental/PythonScriptPlugin/Source/PythonScriptPlugin/Private/K2Node_ExecutePythonScript.h b/Engine/Plugins/Experimental/PythonScriptPlugin/Source/PythonScriptPlugin/Private/K2Node_ExecutePythonScript.h
--- a/Engine/Plugins/Experimental/PythonScriptPlugin/Source/PythonScriptPlugin/Private/K2Node_ExecutePythonScript.h
+++ b/Engine/Plugins/Experimental/PythonScriptPlugin/Source/PythonScriptPlugin/Private/K2Node_ExecutePythonScript.h
@@ -33,6 +33,7 @@ class UK2Node_ExecutePythonScript : public UK2Node_CallFunction
 	//~ Begin UK2Node Interface.
 	virtual void EarlyValidation(class FCompilerResultsLog& MessageLog) const override;
 	virtual void ExpandNode(class FKismetCompilerContext& CompilerContext, UEdGraph* SourceGraph) override;
+	virtual class FNodeHandlingFunctor* CreateNodeHandler(class FKismetCompilerContext& CompilerContext) const override;
 	virtual bool CanPasteHere(const UEdGraph* TargetGraph) const override;
 	virtual bool IsActionFilteredOut(const class FBlueprintActionFilter& Filter) override;
 	virtual void GetMenuActions(FBlueprintActionDatabaseRegistrar& ActionRegistrar) const override;
@@ -48,6 +49,8 @@ private:
 	UEdGraphPin* FindArgumentPinChecked(const FName PinName, EEdGraphPinDirection PinDirection = EGPD_MAX);
 
 private:
+	friend class FKCHandler_ExecutePythonScript;
+
 	/** User-defined input pins */
 	UPROPERTY(EditAnywhere, Category="Arguments")
 	TArray<FName> Inputs;
diff --git a/Engine/Plugins/Experimental/PythonScriptPlugin/Source/PythonScriptPlugin/PythonScriptPlugin.Build.cs b/Engine/Plugins/Experimental/PythonScriptPlugin/Source/PythonScriptPlugin/PythonScriptPlugin.Build.cs
--- a/Engine/Plugins/Experimental/PythonScriptPlugin/Source/PythonScriptPlugin/PythonScriptPlugin.Build.cs
+++ b/Engine/Plugins/Experimental/PythonScriptPlugin/Source/PythonScriptPlugin/PythonScriptPlugin.Build.cs
@@ -42,6 +42,7 @@ namespace UnrealBuildTool.Rules
 
 				PrivateDependencyModuleNames.AddRange(
 					new string[] {
+						"BlueprintGraph",
 						"DesktopPlatform",
 						"EditorStyle",
 						"LevelEditor",
diff --git a/Engine/Source/Editor/BlueprintGraph/Public/CallFunctionHandler.h b/Engine/Source/Editor/BlueprintGraph/Public/CallFunctionHandler.h
--- a/Engine/Source/Editor/BlueprintGraph/Public/CallFunctionHandler.h
+++ b/Engine/Source/Editor/BlueprintGraph/Public/CallFunctionHandler.h
@@ -12,7 +12,7 @@ struct FKismetFunctionContext;
 //////////////////////////////////////////////////////////////////////////
 // FKCHandler_CallFunction
 
-class FKCHandler_CallFunction : public FNodeHandlingFunctor
+class BLUEPRINTGRAPH_API FKCHandler_CallFunction : public FNodeHandlingFunctor
 {
 public:
 	FKCHandler_CallFunction(FKismetCompilerContext& InCompilerContext)
