In Python stub generation, changed a check() to a checkf() to log the name of the UFunction that was detected as a delegate but didn't have the FUNC_Delegate flag.

diff --git a/Engine/Plugins/Experimental/PythonScriptPlugin/Source/PythonScriptPlugin/Private/PyWrapperTypeRegistry.cpp b/Engine/Plugins/Experimental/PythonScriptPlugin/Source/PythonScriptPlugin/Private/PyWrapperTypeRegistry.cpp
--- a/Engine/Plugins/Experimental/PythonScriptPlugin/Source/PythonScriptPlugin/Private/PyWrapperTypeRegistry.cpp
+++ b/Engine/Plugins/Experimental/PythonScriptPlugin/Source/PythonScriptPlugin/Private/PyWrapperTypeRegistry.cpp
@@ -682,7 +682,7 @@ void FPyWrapperTypeRegistry::GenerateWrappedTypesForReferences(const FGeneratedW
 
 	for (const UFunction* DelegateSignature : InGeneratedWrappedTypeReferences.DelegateReferences)
 	{
-		check(DelegateSignature->HasAnyFunctionFlags(FUNC_Delegate));
+		checkf(DelegateSignature->HasAnyFunctionFlags(FUNC_Delegate), TEXT("UFunction '%s' was detected as a delegate but doesn't have the 'FUNC_Delegate' flag"), *DelegateSignature->GetPathName());
 		GenerateWrappedDelegateType(DelegateSignature, GeneratedWrappedTypeReferences, OutDirtyModules, ReferenceGenerationFlags);
 	}
 
