Avoid attempting to generate Python glue for non-C++ types when running the initial glue generation step

#rnx

diff --git a/Engine/Plugins/Experimental/PythonScriptPlugin/Source/PythonScriptPlugin/Private/PyWrapperTypeRegistry.cpp b/Engine/Plugins/Experimental/PythonScriptPlugin/Source/PythonScriptPlugin/Private/PyWrapperTypeRegistry.cpp
--- a/Engine/Plugins/Experimental/PythonScriptPlugin/Source/PythonScriptPlugin/Private/PyWrapperTypeRegistry.cpp
+++ b/Engine/Plugins/Experimental/PythonScriptPlugin/Source/PythonScriptPlugin/Private/PyWrapperTypeRegistry.cpp
@@ -559,11 +559,16 @@ void FPyWrapperTypeRegistry::GenerateWrappedTypes()
 		// Need to use Get rather than ForEach as generating wrapped types 
 		// may generate new objects which breaks the iteration
 		TArray<UObject*> ObjectsToProcess;
-		GetObjectsOfClass(UObject::StaticClass(), ObjectsToProcess);
+		GetObjectsOfClass(UField::StaticClass(), ObjectsToProcess);
 
+		FNameBuilder ObjectPackageName;
 		for (UObject* ObjectToProcess : ObjectsToProcess)
 		{
-			GenerateWrappedTypeForObject(ObjectToProcess, GeneratedWrappedTypeReferences, DirtyModules);
+			ObjectToProcess->GetPackage()->GetFName().ToString(ObjectPackageName);
+			if (FPackageName::IsScriptPackage(ObjectPackageName.ToView()))
+			{
+				GenerateWrappedTypeForObject(ObjectToProcess, GeneratedWrappedTypeReferences, DirtyModules);
+			}
 		}
 
 		GenerateWrappedTypesForReferences(GeneratedWrappedTypeReferences, DirtyModules);
