PythonScriptPlugin: Add import_text / export_text to wrapped structs

diff --git a/Engine/Plugins/Experimental/PythonScriptPlugin/Source/PythonScriptPlugin/Private/PyWrapperStruct.cpp b/Engine/Plugins/Experimental/PythonScriptPlugin/Source/PythonScriptPlugin/Private/PyWrapperStruct.cpp
--- a/Engine/Plugins/Experimental/PythonScriptPlugin/Source/PythonScriptPlugin/Private/PyWrapperStruct.cpp
+++ b/Engine/Plugins/Experimental/PythonScriptPlugin/Source/PythonScriptPlugin/Private/PyWrapperStruct.cpp
@@ -1191,6 +1191,73 @@ PyTypeObject InitializePyWrapperStructType()
 
 			Py_RETURN_NONE;
 		}
+
+		static PyObject* ExportToText(FPyWrapperStruct* InSelf)
+		{
+			if (!FPyWrapperStruct::ValidateInternalState(InSelf))
+			{
+				return nullptr;
+			}
+
+			FString ExportedText;
+			InSelf->ScriptStruct->ExportText(ExportedText, InSelf->StructInstance, InSelf->StructInstance, nullptr, PPF_None, nullptr);
+			return PyConversion::Pythonize(ExportedText);
+		}
+
+		static PyObject* ImportFromText(FPyWrapperStruct* InSelf, PyObject* InArgs, PyObject* InKwds)
+		{
+			if (!FPyWrapperStruct::ValidateInternalState(InSelf))
+			{
+				return nullptr;
+			}
+
+			PyObject* PyContentObj = nullptr;
+			static const char *ArgsKwdList[] = { "content", nullptr };
+			if (!PyArg_ParseTupleAndKeywords(InArgs, InKwds, "O:import_text", (char**)ArgsKwdList, &PyContentObj))
+			{
+				return nullptr;
+			}
+
+			FString TextToImport;
+			if(PyConversion::Nativize(PyContentObj, TextToImport).Failed())
+			{
+				return nullptr;
+			}
+
+			class FErrorPipe : public FOutputDevice
+			{
+			public:
+
+				FPyWrapperStruct* Self;
+				int32 NumErrors;
+
+				FErrorPipe(FPyWrapperStruct* InSelf)
+					: FOutputDevice()
+					, Self(InSelf)
+					, NumErrors(0)
+				{
+				}
+
+				virtual void Serialize(const TCHAR* V, ELogVerbosity::Type Verbosity, const class FName& Category) override
+				{
+					if(Verbosity == ELogVerbosity::Fatal || Verbosity == ELogVerbosity::Error)
+					{
+						if(++NumErrors == 1)
+						{
+							PyUtil::SetPythonError(PyExc_RuntimeError, Self, V);
+						}
+					}
+					else if(Verbosity == ELogVerbosity::Warning)
+					{
+						PyUtil::SetPythonWarning(PyExc_RuntimeWarning, Self, V);
+					}
+				}
+			};
+
+			FErrorPipe ErrorPipe(InSelf);
+			InSelf->ScriptStruct->ImportText(*TextToImport, InSelf->StructInstance, nullptr, PPF_None, &ErrorPipe, InSelf->ScriptStruct->GetName());
+			return PyConversion::Pythonize(ErrorPipe.NumErrors == 0);
+		}
 	};
 
 	static PyMethodDef PyMethods[] = {
@@ -1204,6 +1271,8 @@ PyTypeObject InitializePyWrapperStructType()
 		{ "get_editor_property", PyCFunctionCast(&FMethods::GetEditorProperty), METH_VARARGS | METH_KEYWORDS, "x.get_editor_property(name) -> object -- get the value of any property visible to the editor" },
 		{ "set_editor_property", PyCFunctionCast(&FMethods::SetEditorProperty), METH_VARARGS | METH_KEYWORDS, "x.set_editor_property(name, value, notify_mode=PropertyAccessChangeNotifyMode.DEFAULT) -> None -- set the value of any property visible to the editor, ensuring that the pre/post change notifications are called" },
 		{ "set_editor_properties", PyCFunctionCast(&FMethods::SetEditorProperties), METH_VARARGS, "x.set_editor_properties(property_info) -> None -- set the value of any properties visible to the editor (from a name->value dict), ensuring that the pre/post change notifications are called" },
+		{ "export_text", PyCFunctionCast(&FMethods::ExportToText), METH_NOARGS, "X.export_text() -> str -- exports the content of the Unreal struct of this type" },
+		{ "import_text", PyCFunctionCast(&FMethods::ImportFromText), METH_VARARGS | METH_KEYWORDS, "X.import_text(content) -> bool -- imports the provided string into the Unreal struct" },
 		{ nullptr, nullptr, 0, nullptr }
 	};
 
