PythonScriptPlugin: expose IsPackageExternal() and GetPackage() for object instances wrapped to Python

This exposes UObjectBaseUtility's IsPackageExternal() and GetPackage() as
is_package_external() and get_package(), respectively, for UObject instances
that are wrapped to Python. GetOutermost() simply calls through to GetPackage(),
so the newly added get_package() is equivalent to the existing get_outermost(),
but the former is preferred according to the docstring for GetOutermost().
get_outermost() was left in place so as not to break existing Python code.

diff --git a/Engine/Plugins/Experimental/PythonScriptPlugin/Source/PythonScriptPlugin/Private/PyWrapperObject.cpp b/Engine/Plugins/Experimental/PythonScriptPlugin/Source/PythonScriptPlugin/Private/PyWrapperObject.cpp
--- a/Engine/Plugins/Experimental/PythonScriptPlugin/Source/PythonScriptPlugin/Private/PyWrapperObject.cpp
+++ b/Engine/Plugins/Experimental/PythonScriptPlugin/Source/PythonScriptPlugin/Private/PyWrapperObject.cpp
@@ -604,6 +604,27 @@ PyTypeObject InitializePyWrapperObjectType()
 			return PyConversion::Pythonize(InSelf->ObjectInstance->GetOutermost());
 		}
 
+		static PyObject* IsPackageExternal(FPyWrapperObject* InSelf)
+		{
+			if (!FPyWrapperObject::ValidateInternalState(InSelf))
+			{
+				return nullptr;
+			}
+
+			const bool bIsPackageExternal = InSelf->ObjectInstance->IsPackageExternal();
+			return PyConversion::Pythonize(bIsPackageExternal);
+		}
+
+		static PyObject* GetPackage(FPyWrapperObject* InSelf)
+		{
+			if (!FPyWrapperObject::ValidateInternalState(InSelf))
+			{
+				return nullptr;
+			}
+
+			return PyConversion::Pythonize(InSelf->ObjectInstance->GetPackage());
+		}
+
 		static PyObject* GetName(FPyWrapperObject* InSelf)
 		{
 			if (!FPyWrapperObject::ValidateInternalState(InSelf))
@@ -980,6 +1001,8 @@ PyTypeObject InitializePyWrapperObjectType()
 		{ "get_outer", PyCFunctionCast(&FMethods::GetOuter), METH_NOARGS, "x.get_outer() -> Object -- get the outer object from this instance (if any)" },
 		{ "get_typed_outer", PyCFunctionCast(&FMethods::GetTypedOuter), METH_VARARGS, "x.get_typed_outer(type) -> type() -- get the first outer object of the given type from this instance (if any)" },
 		{ "get_outermost", PyCFunctionCast(&FMethods::GetOutermost), METH_NOARGS, "x.get_outermost() -> Package -- get the outermost object (the package) from this instance" },
+		{ "is_package_external", PyCFunctionCast(&FMethods::IsPackageExternal), METH_NOARGS, "x.is_package_external() -> bool -- returns true if this instance has a different package than its outer's package" },
+		{ "get_package", PyCFunctionCast(&FMethods::GetPackage), METH_NOARGS, "x.get_package() -> Package -- get the package directly associated with this instance" },
 		{ "get_name", PyCFunctionCast(&FMethods::GetName), METH_NOARGS, "x.get_name() -> str -- get the name of this instance" },
 		{ "get_fname", PyCFunctionCast(&FMethods::GetFName), METH_NOARGS, "x.get_fname() -> FName -- get the name of this instance" },
 		{ "get_full_name", PyCFunctionCast(&FMethods::GetFullName), METH_NOARGS, "x.get_full_name() -> str -- get the full name (class name + full path) of this instance" },
