Exposed TFieldPath<T> as 'FieldPath' to Python. That type is used in APIs and types exposed to Python.
  - Added corresponding tests to test_wrapper_type

diff --git a/Engine/Plugins/Experimental/PythonScriptPlugin/Source/PythonScriptPlugin/Private/PyConversion.cpp b/Engine/Plugins/Experimental/PythonScriptPlugin/Source/PythonScriptPlugin/Private/PyConversion.cpp
--- a/Engine/Plugins/Experimental/PythonScriptPlugin/Source/PythonScriptPlugin/Private/PyConversion.cpp
+++ b/Engine/Plugins/Experimental/PythonScriptPlugin/Source/PythonScriptPlugin/Private/PyConversion.cpp
@@ -15,6 +15,7 @@
 #include "PyWrapperFixedArray.h"
 #include "PyWrapperSet.h"
 #include "PyWrapperMap.h"
+#include "PyWrapperFieldPath.h"
 #include "PyWrapperTypeRegistry.h"
 
 #include "Templates/Casts.h"
@@ -365,8 +366,8 @@ FPyConversionResult Nativize(PyObject* PyObj, FText& OutVal, const ESetErrorStat
 {
 	if (PyObject_IsInstance(PyObj, (PyObject*)&PyWrapperTextType) == 1)
 	{
-		FPyWrapperText* PyWrappedName = (FPyWrapperText*)PyObj;
-		OutVal = PyWrappedName->Value;
+		FPyWrapperText* PyWrappedText = (FPyWrapperText*)PyObj;
+		OutVal = PyWrappedText->Value;
 		return FPyConversionResult::Success();
 	}
 
@@ -386,6 +387,25 @@ FPyConversionResult Pythonize(const FText& Val, PyObject*& OutPyObj, const ESetE
 	return FPyConversionResult::Success();
 }
 
+FPyConversionResult Nativize(PyObject* PyObj, FFieldPath& OutVal, const ESetErrorState SetErrorState)
+{
+	if (PyObject_IsInstance(PyObj, (PyObject*)&PyWrapperFieldPathType) == 1)
+	{
+		FPyWrapperFieldPath* PyWrappedFieldPath = (FPyWrapperFieldPath*)PyObj;
+		OutVal = PyWrappedFieldPath->Value;
+		return FPyConversionResult::Success();
+	}
+
+	PYCONVERSION_RETURN(FPyConversionResult::Failure(), TEXT("Nativize"), *FString::Printf(TEXT("Cannot nativize '%s' as 'FieldPath'"), *PyUtil::GetFriendlyTypename(PyObj)));
+}
+
+FPyConversionResult Pythonize(const FFieldPath& Val, PyObject*& OutPyObj, const ESetErrorState SetErrorState)
+{
+	OutPyObj = (PyObject*)FPyWrapperFieldPathFactory::Get().CreateInstance(Val);
+	return FPyConversionResult::Success();
+}
+
+
 FPyConversionResult Nativize(PyObject* PyObj, void*& OutVal, const ESetErrorState SetErrorState)
 {
 	// CObject was removed in Python 3.2
@@ -671,7 +691,7 @@ FPyConversionResult NativizeProperty_Direct(PyObject* PyObj, const FProperty* Pr
 	PYCONVERSION_RETURN(RESULT, TEXT("NativizeProperty"), *FString::Printf(TEXT("Cannot nativize '%s' as '%s' (%s)"), *PyUtil::GetFriendlyTypename(PyObj), *Prop->GetName(), *Prop->GetClass()->GetName()))
 
 #define NATIVIZE_SETTER_PROPERTY(PROPTYPE)											\
-	if (auto* CastProp = CastField<PROPTYPE>(Prop))										\
+	if (const PROPTYPE* CastProp = CastField<PROPTYPE>(Prop))								\
 	{																				\
 		PROPTYPE::TCppType NewValue;												\
 		const FPyConversionResult Result = Nativize(PyObj, NewValue, SetErrorState);\
@@ -688,7 +708,7 @@ FPyConversionResult NativizeProperty_Direct(PyObject* PyObj, const FProperty* Pr
 	}
 
 #define NATIVIZE_INLINE_PROPERTY(PROPTYPE)											\
-	if (auto* CastProp = CastField<PROPTYPE>(Prop))										\
+	if (const PROPTYPE* CastProp = CastField<PROPTYPE>(Prop))								\
 	{																				\
 		PROPTYPE::TCppType NewValue;												\
 		const FPyConversionResult Result = Nativize(PyObj, NewValue, SetErrorState);\
@@ -717,8 +737,9 @@ FPyConversionResult NativizeProperty_Direct(PyObject* PyObj, const FProperty* Pr
 	NATIVIZE_INLINE_PROPERTY(FStrProperty);
 	NATIVIZE_INLINE_PROPERTY(FNameProperty);
 	NATIVIZE_INLINE_PROPERTY(FTextProperty);
+	NATIVIZE_INLINE_PROPERTY(FFieldPathProperty);
 
-	if (auto* CastProp = CastField<FByteProperty>(Prop))
+	if (const FByteProperty* CastProp = CastField<FByteProperty>(Prop))
 	{
 		uint8 NewValue = 0;
 		FPyConversionResult Result = FPyConversionResult::Failure();
@@ -753,7 +774,7 @@ FPyConversionResult NativizeProperty_Direct(PyObject* PyObj, const FProperty* Pr
 		PYCONVERSION_PROPERTY_RETURN(Result);
 	}
 
-	if (auto* CastProp = CastField<FEnumProperty>(Prop))
+	if (const FEnumProperty* CastProp = CastField<FEnumProperty>(Prop))
 	{
 		FPyConversionResult Result = FPyConversionResult::Failure();
 
@@ -782,7 +803,7 @@ FPyConversionResult NativizeProperty_Direct(PyObject* PyObj, const FProperty* Pr
 		PYCONVERSION_PROPERTY_RETURN(Result);
 	}
 
-	if (auto* CastProp = CastField<FClassProperty>(Prop))
+	if (const FClassProperty* CastProp = CastField<FClassProperty>(Prop))
 	{
 		UClass* NewValue = nullptr;
 		const FPyConversionResult Result = NativizeClass(PyObj, NewValue, CastProp->MetaClass, SetErrorState);
@@ -797,7 +818,7 @@ FPyConversionResult NativizeProperty_Direct(PyObject* PyObj, const FProperty* Pr
 		PYCONVERSION_PROPERTY_RETURN(Result);
 	}
 
-	if (auto* CastProp = CastField<FSoftClassProperty>(Prop))
+	if (const FSoftClassProperty* CastProp = CastField<FSoftClassProperty>(Prop))
 	{
 		UClass* NewValue = nullptr;
 		const FPyConversionResult Result = NativizeClass(PyObj, NewValue, CastProp->MetaClass, SetErrorState);
@@ -812,7 +833,7 @@ FPyConversionResult NativizeProperty_Direct(PyObject* PyObj, const FProperty* Pr
 		PYCONVERSION_PROPERTY_RETURN(Result);
 	}
 
-	if (auto* CastProp = CastField<FObjectPropertyBase>(Prop))
+	if (const FObjectPropertyBase* CastProp = CastField<FObjectPropertyBase>(Prop))
 	{
 		UObject* NewValue = nullptr;
 		const FPyConversionResult Result = NativizeObject(PyObj, NewValue, CastProp->PropertyClass, SetErrorState);
@@ -827,7 +848,7 @@ FPyConversionResult NativizeProperty_Direct(PyObject* PyObj, const FProperty* Pr
 		PYCONVERSION_PROPERTY_RETURN(Result);
 	}
 
-	if (auto* CastProp = CastField<FInterfaceProperty>(Prop))
+	if (const FInterfaceProperty* CastProp = CastField<FInterfaceProperty>(Prop))
 	{
 		UObject* NewValue = nullptr;
 		const FPyConversionResult Result = NativizeObject(PyObj, NewValue, CastProp->InterfaceClass, SetErrorState);
@@ -842,7 +863,7 @@ FPyConversionResult NativizeProperty_Direct(PyObject* PyObj, const FProperty* Pr
 		PYCONVERSION_PROPERTY_RETURN(Result);
 	}
 
-	if (auto* CastProp = CastField<FStructProperty>(Prop))
+	if (const FStructProperty* CastProp = CastField<FStructProperty>(Prop))
 	{
 		FPyConversionResult Result = FPyConversionResult::Failure();
 		PyTypeObject* PyStructType = FPyWrapperTypeRegistry::Get().GetWrappedStructType(CastProp->Struct);
@@ -865,7 +886,7 @@ FPyConversionResult NativizeProperty_Direct(PyObject* PyObj, const FProperty* Pr
 		PYCONVERSION_PROPERTY_RETURN(Result);
 	}
 
-	if (auto* CastProp = CastField<FDelegateProperty>(Prop))
+	if (const FDelegateProperty* CastProp = CastField<FDelegateProperty>(Prop))
 	{
 		FPyConversionResult Result = FPyConversionResult::Failure();
 		PyTypeObject* PyDelegateType = FPyWrapperTypeRegistry::Get().GetWrappedDelegateType(CastProp->SignatureFunction);
@@ -880,7 +901,7 @@ FPyConversionResult NativizeProperty_Direct(PyObject* PyObj, const FProperty* Pr
 		PYCONVERSION_PROPERTY_RETURN(Result);
 	}
 
-	if (auto* CastProp = CastField<FMulticastDelegateProperty>(Prop))
+	if (const FMulticastDelegateProperty* CastProp = CastField<FMulticastDelegateProperty>(Prop))
 	{
 		FPyConversionResult Result = FPyConversionResult::Failure();
 		PyTypeObject* PyDelegateType = FPyWrapperTypeRegistry::Get().GetWrappedDelegateType(CastProp->SignatureFunction);
@@ -895,7 +916,7 @@ FPyConversionResult NativizeProperty_Direct(PyObject* PyObj, const FProperty* Pr
 		PYCONVERSION_PROPERTY_RETURN(Result);
 	}
 
-	if (auto* CastProp = CastField<FArrayProperty>(Prop))
+	if (const FArrayProperty* CastProp = CastField<FArrayProperty>(Prop))
 	{
 		FPyConversionResult Result = FPyConversionResult::Failure();
 		FPyWrapperArrayPtr PyArray = FPyWrapperArrayPtr::StealReference(FPyWrapperArray::CastPyObject(PyObj, &PyWrapperArrayType, CastProp->Inner, &Result));
@@ -909,7 +930,7 @@ FPyConversionResult NativizeProperty_Direct(PyObject* PyObj, const FProperty* Pr
 		PYCONVERSION_PROPERTY_RETURN(Result);
 	}
 
-	if (auto* CastProp = CastField<FSetProperty>(Prop))
+	if (const FSetProperty* CastProp = CastField<FSetProperty>(Prop))
 	{
 		FPyConversionResult Result = FPyConversionResult::Failure();
 		FPyWrapperSetPtr PySet = FPyWrapperSetPtr::StealReference(FPyWrapperSet::CastPyObject(PyObj, &PyWrapperSetType, CastProp->ElementProp, &Result));
@@ -923,7 +944,7 @@ FPyConversionResult NativizeProperty_Direct(PyObject* PyObj, const FProperty* Pr
 		PYCONVERSION_PROPERTY_RETURN(Result);
 	}
 
-	if (auto* CastProp = CastField<FMapProperty>(Prop))
+	if (const FMapProperty* CastProp = CastField<FMapProperty>(Prop))
 	{
 		FPyConversionResult Result = FPyConversionResult::Failure();
 		FPyWrapperMapPtr PyMap = FPyWrapperMapPtr::StealReference(FPyWrapperMap::CastPyObject(PyObj, &PyWrapperMapType, CastProp->KeyProp, CastProp->ValueProp, &Result));
@@ -954,7 +975,7 @@ FPyConversionResult PythonizeProperty_Direct(const FProperty* Prop, const void*
 	OwnerContext.AssertValidConversionMethod(ConversionMethod);
 
 #define PYTHONIZE_GETTER_PROPERTY(PROPTYPE)											\
-	if (auto* CastProp = CastField<PROPTYPE>(Prop))										\
+	if (const PROPTYPE* CastProp = CastField<PROPTYPE>(Prop))										\
 	{																				\
 		auto&& Value = CastProp->GetPropertyValue(ValueAddr);						\
 		PYCONVERSION_PROPERTY_RETURN(Pythonize(Value, OutPyObj, SetErrorState));	\
@@ -973,8 +994,9 @@ FPyConversionResult PythonizeProperty_Direct(const FProperty* Prop, const void*
 	PYTHONIZE_GETTER_PROPERTY(FStrProperty);
 	PYTHONIZE_GETTER_PROPERTY(FNameProperty);
 	PYTHONIZE_GETTER_PROPERTY(FTextProperty);
+	PYTHONIZE_GETTER_PROPERTY(FFieldPathProperty);
 
-	if (auto* CastProp = CastField<FByteProperty>(Prop))
+	if (const FByteProperty* CastProp = CastField<FByteProperty>(Prop))
 	{
 		const uint8 Value = CastProp->GetPropertyValue(ValueAddr);
 		if (CastProp->Enum)
@@ -987,31 +1009,31 @@ FPyConversionResult PythonizeProperty_Direct(const FProperty* Prop, const void*
 		}
 	}
 
-	if (auto* CastProp = CastField<FEnumProperty>(Prop))
+	if (const FEnumProperty* CastProp = CastField<FEnumProperty>(Prop))
 	{
 		FNumericProperty* EnumInternalProp = CastProp->GetUnderlyingProperty();
 		PYCONVERSION_PROPERTY_RETURN(PythonizeEnumEntry(EnumInternalProp ? EnumInternalProp->GetSignedIntPropertyValue(ValueAddr) : 0, CastProp->GetEnum(), OutPyObj, SetErrorState));
 	}
 
-	if (auto* CastProp = CastField<FClassProperty>(Prop))
+	if (const FClassProperty* CastProp = CastField<FClassProperty>(Prop))
 	{
 		UClass* Value = Cast<UClass>(CastProp->LoadObjectPropertyValue(ValueAddr));
 		PYCONVERSION_PROPERTY_RETURN(PythonizeClass(Value, OutPyObj, SetErrorState));
 	}
 
-	if (auto* CastProp = CastField<FSoftClassProperty>(Prop))
+	if (const FSoftClassProperty* CastProp = CastField<FSoftClassProperty>(Prop))
 	{
 		UClass* Value = Cast<UClass>(CastProp->LoadObjectPropertyValue(ValueAddr));
 		PYCONVERSION_PROPERTY_RETURN(PythonizeClass(Value, OutPyObj, SetErrorState));
 	}
 
-	if (auto* CastProp = CastField<FObjectPropertyBase>(Prop))
+	if (const FObjectPropertyBase* CastProp = CastField<FObjectPropertyBase>(Prop))
 	{
 		UObject* Value = CastProp->LoadObjectPropertyValue(ValueAddr);
 		PYCONVERSION_PROPERTY_RETURN(Pythonize(Value, OutPyObj, SetErrorState));
 	}
 
-	if (auto* CastProp = CastField<FInterfaceProperty>(Prop))
+	if (const FInterfaceProperty* CastProp = CastField<FInterfaceProperty>(Prop))
 	{
 		UObject* Value = CastProp->GetPropertyValue(ValueAddr).GetObject();
 		if (Value)
@@ -1026,20 +1048,20 @@ FPyConversionResult PythonizeProperty_Direct(const FProperty* Prop, const void*
 		return FPyConversionResult::Success();
 	}
 
-	if (auto* CastProp = CastField<FStructProperty>(Prop))
+	if (const FStructProperty* CastProp = CastField<FStructProperty>(Prop))
 	{
 		OutPyObj = (PyObject*)FPyWrapperStructFactory::Get().CreateInstance(CastProp->Struct, (void*)ValueAddr, OwnerContext, ConversionMethod);
 		return FPyConversionResult::Success();
 	}
 
-	if (auto* CastProp = CastField<FDelegateProperty>(Prop))
+	if (const FDelegateProperty* CastProp = CastField<FDelegateProperty>(Prop))
 	{
 		const FScriptDelegate* Value = CastProp->GetPropertyValuePtr(ValueAddr);
 		OutPyObj = (PyObject*)FPyWrapperDelegateFactory::Get().CreateInstance(CastProp->SignatureFunction, (FScriptDelegate*)Value, OwnerContext, ConversionMethod);
 		return FPyConversionResult::Success();
 	}
 
-	if (auto* CastProp = CastField<FMulticastDelegateProperty>(Prop))
+	if (const FMulticastDelegateProperty* CastProp = CastField<FMulticastDelegateProperty>(Prop))
 	{
 		const FMulticastScriptDelegate* Value = CastProp->GetMulticastDelegate(ValueAddr);
 		OutPyObj = (PyObject*)FPyWrapperMulticastDelegateFactory::Get().CreateInstance(CastProp->SignatureFunction, (FMulticastScriptDelegate*)Value, OwnerContext, ConversionMethod);
@@ -1047,26 +1069,26 @@ FPyConversionResult PythonizeProperty_Direct(const FProperty* Prop, const void*
 	}
 
 	// TODO: We're just not handling sparse delegates at this time
-	/*if (auto* CastProp = CastField<FMulticastSparseDelegateProperty>(Prop))
+	/*if (const FMulticastSparseDelegateProperty* CastProp = CastField<FMulticastSparseDelegateProperty>(Prop))
 	{
 		const FMulticastScriptDelegate* Value = CastProp->GetMulticastDelegate(ValueAddr);
 		OutPyObj = (PyObject*)FPyWrapperMulticastDelegateFactory::Get().CreateInstance(CastProp->SignatureFunction, (FMulticastScriptDelegate*)Value, OwnerContext, ConversionMethod);
 		return FPyConversionResult::Success();
 	}*/
 
-	if (auto* CastProp = CastField<FArrayProperty>(Prop))
+	if (const FArrayProperty* CastProp = CastField<FArrayProperty>(Prop))
 	{
 		OutPyObj = (PyObject*)FPyWrapperArrayFactory::Get().CreateInstance((void*)ValueAddr, CastProp, OwnerContext, ConversionMethod);
 		return FPyConversionResult::Success();
 	}
 
-	if (auto* CastProp = CastField<FSetProperty>(Prop))
+	if (const FSetProperty* CastProp = CastField<FSetProperty>(Prop))
 	{
 		OutPyObj = (PyObject*)FPyWrapperSetFactory::Get().CreateInstance((void*)ValueAddr, CastProp, OwnerContext, ConversionMethod);
 		return FPyConversionResult::Success();
 	}
 
-	if (auto* CastProp = CastField<FMapProperty>(Prop))
+	if (const FMapProperty* CastProp = CastField<FMapProperty>(Prop))
 	{
 		OutPyObj = (PyObject*)FPyWrapperMapFactory::Get().CreateInstance((void*)ValueAddr, CastProp, OwnerContext, ConversionMethod);
 		return FPyConversionResult::Success();
diff --git a/Engine/Plugins/Experimental/PythonScriptPlugin/Source/PythonScriptPlugin/Private/PyCore.cpp b/Engine/Plugins/Experimental/PythonScriptPlugin/Source/PythonScriptPlugin/Private/PyCore.cpp
--- a/Engine/Plugins/Experimental/PythonScriptPlugin/Source/PythonScriptPlugin/Private/PyCore.cpp
+++ b/Engine/Plugins/Experimental/PythonScriptPlugin/Source/PythonScriptPlugin/Private/PyCore.cpp
@@ -18,6 +18,7 @@
 #include "PyWrapperFixedArray.h"
 #include "PyWrapperSet.h"
 #include "PyWrapperMap.h"
+#include "PyWrapperFieldPath.h"
 #include "PyWrapperMath.h"
 
 #include "PythonScriptPlugin.h"
@@ -1894,6 +1895,7 @@ void InitializeModule()
 	InitializePyWrapperFixedArray(NativePythonModule);
 	InitializePyWrapperSet(NativePythonModule);
 	InitializePyWrapperMap(NativePythonModule);
+	InitializePyWrapperFieldPath(NativePythonModule);
 	InitializePyWrapperMath(NativePythonModule);
 
 	FPyWrapperTypeRegistry::Get().RegisterNativePythonModule(MoveTemp(NativePythonModule));
diff --git a/Engine/Plugins/Experimental/PythonScriptPlugin/Source/PythonScriptPlugin/Private/PyGenUtil.cpp b/Engine/Plugins/Experimental/PythonScriptPlugin/Source/PythonScriptPlugin/Private/PyGenUtil.cpp
--- a/Engine/Plugins/Experimental/PythonScriptPlugin/Source/PythonScriptPlugin/Private/PyGenUtil.cpp
+++ b/Engine/Plugins/Experimental/PythonScriptPlugin/Source/PythonScriptPlugin/Private/PyGenUtil.cpp
@@ -2229,6 +2229,22 @@ void PythonizeValueImpl(const FProperty* InProp, const void* InPropValue, const
 				? FString::Printf(TEXT("%sText(\"%s\")"), UnrealNamespace, **TextStrValue)
 				: FString::Printf(TEXT("\"%s\""), **TextStrValue);
 		}
+		else if (const FFieldPathProperty* FieldPathProp =  CastField<const FFieldPathProperty>(InProp))
+		{
+			const FString FieldPathStrValue = FieldPathProp->GetPropertyValue(PropArrValue).ToString();
+			if (FieldPathStrValue.IsEmpty())
+			{
+				OutPythonDefaultValue += bUseStrictTyping
+					? FString::Printf(TEXT("%sFieldPath()"), UnrealNamespace)
+					: FString::Printf(TEXT("FieldPath()"));
+			}
+			else // Use the path string as default value.
+			{
+				OutPythonDefaultValue += bUseStrictTyping
+					? FString::Printf(TEXT("%sFieldPath(\"%s\")"), UnrealNamespace, *GetPropertyPythonName(FieldPathProp), *FieldPathStrValue)
+					: FString::Printf(TEXT("FieldPath(\"%s\")"), *GetPropertyPythonName(FieldPathProp), *FieldPathStrValue);
+			}
+		}
 		else if (const FObjectPropertyBase* ObjProp = CastField<const FObjectPropertyBase>(InProp))
 		{
 			OutPythonDefaultValue += TEXT("None");
@@ -2298,7 +2314,7 @@ void PythonizeValueImpl(const FProperty* InProp, const void* InPropValue, const
 			OutPythonDefaultValue += bUseStrictTyping
 				? FString::Printf(TEXT("%sMap.cast(%s, %s, {"), UnrealNamespace, *GetPropertyTypePythonName(MapProperty->KeyProp, bIncludeUnrealNamespace, bIsForDocString), *GetPropertyTypePythonName(MapProperty->ValueProp, bIncludeUnrealNamespace, bIsForDocString))
 				: TEXT("{");
-		{
+			{
 				FScriptMapHelper ScriptMapHelper(MapProperty, PropArrValue);
 				for (int32 SparseElementIndex = 0, ElementIndex = 0; SparseElementIndex < ScriptMapHelper.GetMaxIndex(); ++SparseElementIndex)
 				{
@@ -2985,6 +3001,8 @@ FString GetPropertyTypePythonName(const FProperty* InProp, const bool InIncludeU
 	GET_PROPERTY_TYPE(FStrProperty, TEXT("str"))
 	GET_PROPERTY_TYPE(FNameProperty, InIncludeUnrealNamespace ? TEXT("unreal.Name") : TEXT("Name"))
 	GET_PROPERTY_TYPE(FTextProperty, InIncludeUnrealNamespace ? TEXT("unreal.Text") : TEXT("Text"))
+	GET_PROPERTY_TYPE(FFieldPathProperty, InIncludeUnrealNamespace ? TEXT("unreal.FieldPath") : TEXT("FieldPath"))
+
 	if (const FByteProperty* ByteProp = CastField<const FByteProperty>(InProp))
 	{
 		if (ByteProp->Enum)
diff --git a/Engine/Plugins/Experimental/PythonScriptPlugin/Source/PythonScriptPlugin/Private/PyTest.cpp b/Engine/Plugins/Experimental/PythonScriptPlugin/Source/PythonScriptPlugin/Private/PyTest.cpp
--- a/Engine/Plugins/Experimental/PythonScriptPlugin/Source/PythonScriptPlugin/Private/PyTest.cpp
+++ b/Engine/Plugins/Experimental/PythonScriptPlugin/Source/PythonScriptPlugin/Private/PyTest.cpp
@@ -117,6 +117,11 @@ int32 UPyTestObject::FuncTakingPyTestDelegate(const FPyTestDelegate& InDelegate,
 	return InDelegate.IsBound() ? InDelegate.Execute(InValue) : INDEX_NONE;
 }
 
+void UPyTestObject::FuncTakingFieldPath(const TFieldPath<FProperty>& InFieldPath)
+{
+	FieldPath = InFieldPath;
+}
+
 int32 UPyTestObject::DelegatePropertyCallback(const int32 InValue) const
 {
 	if (InValue != Int)
@@ -156,6 +161,11 @@ TMap<int32, bool> UPyTestObject::ReturnMap()
 	return TmpMap;
 }
 
+TFieldPath<FProperty> UPyTestObject::ReturnFieldPath()
+{
+	return TFieldPath<FProperty>(UPyTestObject::StaticClass()->FindPropertyByName(TEXT("FieldPath")));
+}
+
 void UPyTestObject::EmitScriptError()
 {
 	FFrame::KismetExecutionMessage(TEXT("EmitScriptError was called"), ELogVerbosity::Error);
diff --git a/Engine/Plugins/Experimental/PythonScriptPlugin/Source/PythonScriptPlugin/Private/PyTest.h b/Engine/Plugins/Experimental/PythonScriptPlugin/Source/PythonScriptPlugin/Private/PyTest.h
--- a/Engine/Plugins/Experimental/PythonScriptPlugin/Source/PythonScriptPlugin/Private/PyTest.h
+++ b/Engine/Plugins/Experimental/PythonScriptPlugin/Source/PythonScriptPlugin/Private/PyTest.h
@@ -59,6 +59,12 @@ public:
 	UPROPERTY(EditAnywhere, BlueprintReadWrite, Category = "Python|Internal")
 	FText Text;
 
+	UPROPERTY(EditAnywhere, BlueprintReadWrite, Category = "Python|Internal")
+	TFieldPath<FProperty> FieldPath;
+
+	UPROPERTY(EditAnywhere, BlueprintReadWrite, Category = "Python|Internal")
+	TFieldPath<FStructProperty> StructFieldPath;
+
 	UPROPERTY(EditAnywhere, BlueprintReadWrite, Category = "Python|Internal")
 	TArray<FString> StringArray;
 
@@ -146,6 +152,12 @@ public:
 	UPROPERTY(EditAnywhere, BlueprintReadWrite, Category = "Python|Internal")
 	FText Text;
 
+	UPROPERTY(EditAnywhere, BlueprintReadWrite, Category = "Python|Internal")
+	TFieldPath<FProperty> FieldPath = TFieldPath<FProperty>(FPyTestStruct::StaticStruct()->FindPropertyByName(TEXT("StringArray")));
+
+	UPROPERTY(EditAnywhere, BlueprintReadWrite, Category = "Python|Internal")
+	TFieldPath<FStructProperty> StructFieldPath = TFieldPath<FStructProperty>(CastField<FStructProperty>(UPyTestObject::GetClass()->FindPropertyByName(TEXT("Struct"))));
+
 	UPROPERTY(EditAnywhere, BlueprintReadWrite, Category = "Python|Internal")
 	TArray<FString> StringArray;
 
@@ -209,6 +221,9 @@ public:
 	UFUNCTION(BlueprintCallable, Category = "Python|Internal")
 	int32 FuncTakingPyTestDelegate(const FPyTestDelegate& InDelegate, const int32 InValue) const;
 
+	UFUNCTION(BlueprintCallable, Category = "Python|Internal")
+	void FuncTakingFieldPath(const TFieldPath<FProperty>& InFieldPath); // UHT couldn't parse any default value for the FieldPath.
+
 	UFUNCTION(BlueprintCallable, Category = "Python|Internal")
 	int32 DelegatePropertyCallback(const int32 InValue) const;
 
@@ -224,6 +239,9 @@ public:
 	UFUNCTION(BlueprintCallable, Category = "Python|Internal")
 	static TMap<int32, bool> ReturnMap();
 
+	UFUNCTION(BlueprintCallable, Category = "Python|Internal")
+	static TFieldPath<FProperty> ReturnFieldPath();
+
 	UFUNCTION(BlueprintCallable, Category = "Python|Internal")
 	static void EmitScriptError();
 
diff --git a/Engine/Plugins/Experimental/PythonScriptPlugin/Source/PythonScriptPlugin/Private/PyUtil.cpp b/Engine/Plugins/Experimental/PythonScriptPlugin/Source/PythonScriptPlugin/Private/PyUtil.cpp
--- a/Engine/Plugins/Experimental/PythonScriptPlugin/Source/PythonScriptPlugin/Private/PyUtil.cpp
+++ b/Engine/Plugins/Experimental/PythonScriptPlugin/Source/PythonScriptPlugin/Private/PyUtil.cpp
@@ -15,6 +15,7 @@
 #include "PyWrapperFixedArray.h"
 #include "PyWrapperSet.h"
 #include "PyWrapperMap.h"
+#include "PyWrapperFieldPath.h"
 #include "PyWrapperTypeRegistry.h"
 
 #include "Misc/Paths.h"
@@ -305,6 +306,12 @@ bool CalculatePropertyDef(PyTypeObject* InPyType, FPropertyDef& OutPropertyDef)
 		return true;
 	}
 
+	if (PyObject_IsSubclass((PyObject*)InPyType, (PyObject*)&PyWrapperFieldPathType) == 1)
+	{
+		OutPropertyDef.PropertyClass = FFieldPathProperty::StaticClass();
+		return true;
+	}
+
 	if (PyObject_IsSubclass((PyObject*)InPyType, (PyObject*)&PyUnicode_Type) == 1)
 	{
 		OutPropertyDef.PropertyClass = FStrProperty::StaticClass();
diff --git a/Engine/Plugins/Experimental/PythonScriptPlugin/Source/PythonScriptPlugin/Private/PyWrapperFieldPath.cpp b/Engine/Plugins/Experimental/PythonScriptPlugin/Source/PythonScriptPlugin/Private/PyWrapperFieldPath.cpp
new file mode 100644
--- /dev/null
+++ b/Engine/Plugins/Experimental/PythonScriptPlugin/Source/PythonScriptPlugin/Private/PyWrapperFieldPath.cpp
@@ -0,0 +1,216 @@
+// Copyright Epic Games, Inc. All Rights Reserved.
+
+#include "PyWrapperFieldPath.h"
+#include "PyWrapperTypeRegistry.h"
+#include "PyConversion.h"
+#include "PyReferenceCollector.h"
+#include "UObject/UnrealType.h"
+#include "UObject/Package.h"
+#include "UObject/PropertyPortFlags.h"
+
+#if WITH_PYTHON
+
+void InitializePyWrapperFieldPath(PyGenUtil::FNativePythonModule& ModuleInfo)
+{
+	if (PyType_Ready(&PyWrapperFieldPathType) == 0)
+	{
+		ModuleInfo.AddType(&PyWrapperFieldPathType);
+	}
+}
+
+void FPyWrapperFieldPath::InitValue(FPyWrapperFieldPath* InSelf, const FFieldPath& InValue)
+{
+	Super::InitValue(InSelf, InValue);
+	FPyWrapperFieldPathFactory::Get().MapInstance(InSelf->Value, InSelf);
+}
+
+
+void FPyWrapperFieldPath::DeinitValue(FPyWrapperFieldPath* InSelf)
+{
+	FPyWrapperFieldPathFactory::Get().UnmapInstance(InSelf->Value, Py_TYPE(InSelf));
+	Super::DeinitValue(InSelf);
+}
+
+FPyWrapperFieldPath* FPyWrapperFieldPath::CastPyObject(PyObject* InPyObject, FPyConversionResult* OutCastResult)
+{
+	SetOptionalPyConversionResult(FPyConversionResult::Failure(), OutCastResult);
+
+	if (PyObject_IsInstance(InPyObject, (PyObject*)&PyWrapperFieldPathType) == 1)
+	{
+		SetOptionalPyConversionResult(FPyConversionResult::Success(), OutCastResult);
+
+		Py_INCREF(InPyObject);
+		return (FPyWrapperFieldPath*)InPyObject;
+	}
+
+	return nullptr;
+}
+
+FPyWrapperFieldPath* FPyWrapperFieldPath::CastPyObject(PyObject* InPyObject, PyTypeObject* InType, FPyConversionResult* OutCastResult)
+{
+	SetOptionalPyConversionResult(FPyConversionResult::Failure(), OutCastResult);
+
+	if (PyObject_IsInstance(InPyObject, (PyObject*)InType) == 1 && (InType == &PyWrapperFieldPathType || PyObject_IsInstance(InPyObject, (PyObject*)&PyWrapperFieldPathType) == 1))
+	{
+		SetOptionalPyConversionResult(Py_TYPE(InPyObject) == InType ? FPyConversionResult::Success() : FPyConversionResult::SuccessWithCoercion(), OutCastResult);
+
+		Py_INCREF(InPyObject);
+		return (FPyWrapperFieldPath*)InPyObject;
+	}
+
+	{
+		FFieldPath InitValue;
+		if (PyConversion::Nativize(InPyObject, InitValue, PyConversion::ESetErrorState::No))
+		{
+			FPyWrapperFieldPathPtr NewFiedPath = FPyWrapperFieldPathPtr::StealReference(FPyWrapperFieldPath::New(InType));
+			if (NewFiedPath)
+			{
+				if (FPyWrapperFieldPath::Init(NewFiedPath, InitValue) != 0)
+				{
+					return nullptr;
+				}
+			}
+			SetOptionalPyConversionResult(FPyConversionResult::SuccessWithCoercion(), OutCastResult);
+			return NewFiedPath.Release();
+		}
+	}
+
+	{
+		FString InitPath;
+		if (PyConversion::Nativize(InPyObject, InitPath, PyConversion::ESetErrorState::No))
+		{
+			FPyWrapperFieldPathPtr NewFiedPath = FPyWrapperFieldPathPtr::StealReference(FPyWrapperFieldPath::New(InType));
+			if (NewFiedPath)
+			{
+				FFieldPath InitValue;
+				InitValue.Generate(*InitPath);
+				if (FPyWrapperFieldPath::Init(NewFiedPath, InitValue) != 0)
+				{
+					return nullptr;
+				}
+			}
+			SetOptionalPyConversionResult(FPyConversionResult::SuccessWithCoercion(), OutCastResult);
+			return NewFiedPath.Release();
+		}
+	}
+
+	return nullptr;
+}
+
+PyTypeObject InitializePyWrapperFieldPathType()
+{
+	struct FFuncs
+	{
+		static int Init(FPyWrapperFieldPath* InSelf, PyObject* InArgs, PyObject* InKwds)
+		{
+			PyObject* PyObj = nullptr;
+			if (!PyArg_ParseTuple(InArgs, "|O:call", &PyObj))
+			{
+				return -1;
+			}
+
+			FFieldPath InitValue;
+			FString StrValue;
+			if (PyObj && PyConversion::Nativize(PyObj, StrValue, PyConversion::ESetErrorState::No)) // Init from a string.
+			{
+				InitValue.Generate(*StrValue);
+			}
+			else if (PyObj && !PyConversion::Nativize(PyObj, InitValue)) // Init from another FieldPath.
+			{
+				PyUtil::SetPythonError(PyExc_TypeError, InSelf, *FString::Printf(TEXT("Failed to convert init argument '%s' to 'FieldPath'"), *PyUtil::GetFriendlyTypename(PyObj)));
+				return -1;
+			}
+			// else: Default initialization. "f = unreal.FieldPath()".
+
+			return FPyWrapperFieldPath::Init(InSelf, InitValue);
+		}
+
+		static PyObject* Str(FPyWrapperFieldPath* InSelf)
+		{
+			return PyUnicode_FromFormat(TCHAR_TO_UTF8(*InSelf->Value.ToString()));
+		}
+
+		static PyObject* Repr(FPyWrapperFieldPath* InSelf)
+		{
+			return PyUnicode_FromString(TCHAR_TO_UTF8(*FString::Printf(TEXT("FieldPath(\"%s\")"), *InSelf->Value.ToString())));
+		}
+
+		static PyObject* RichCmp(FPyWrapperFieldPath* InSelf, PyObject* InOther, int InOp)
+		{
+			FFieldPath Other;
+			if (!PyConversion::Nativize(InOther, Other, PyConversion::ESetErrorState::No))
+			{
+				Py_INCREF(Py_NotImplemented);
+				return Py_NotImplemented;
+			}
+
+			if (InOp != Py_EQ && InOp != Py_NE)
+			{
+				PyUtil::SetPythonError(PyExc_TypeError, InSelf, TEXT("Only == and != comparison is supported"));
+				return nullptr;
+			}
+
+			bool bIdentical = InSelf->Value == Other;
+			return PyBool_FromLong(InOp == Py_EQ ? bIdentical : !bIdentical);
+		}
+
+		static PyUtil::FPyHashType Hash(FPyWrapperFieldPath* InSelf)
+		{
+			uint32 FieldPathHash = GetTypeHash(InSelf->Value);
+			return FieldPathHash != -1 ? FieldPathHash : 0;
+		}
+	};
+
+	struct FMethods
+	{
+		static PyObject* Cast(PyTypeObject* InType, PyObject* InArgs)
+		{
+			PyObject* PyObj = nullptr;
+			if (PyArg_ParseTuple(InArgs, "O:call", &PyObj))
+			{
+				PyObject* PyCastResult = (PyObject*)FPyWrapperFieldPath::CastPyObject(PyObj, InType);
+				if (!PyCastResult)
+				{
+					PyUtil::SetPythonError(PyExc_TypeError, InType, *FString::Printf(TEXT("Cannot cast type '%s' to '%s'"), *PyUtil::GetFriendlyTypename(PyObj), *PyUtil::GetFriendlyTypename(InType)));
+				}
+				return PyCastResult;
+			}
+
+			return nullptr;
+		}
+
+		static PyObject* Copy(FPyWrapperFieldPath* InSelf)
+		{
+			return PyConversion::Pythonize(InSelf->Value);
+		}
+
+		static PyObject* IsValid(FPyWrapperFieldPath* InSelf)
+		{
+			return PyConversion::Pythonize(InSelf->Value.TryToResolvePath(nullptr) != nullptr);
+		}
+	};
+
+	static PyMethodDef PyMethods[] = {
+		{ "cast", PyCFunctionCast(&FMethods::Cast), METH_VARARGS | METH_CLASS, "X.cast(obj: Any) -> FieldPath -- cast the given object to this Unreal field path type" },
+		{ "__copy__", PyCFunctionCast(&FMethods::Copy), METH_NOARGS, "x.__copy__() -> FieldPath -- copy this Unreal field path" },
+		{ "copy", PyCFunctionCast(&FMethods::Copy), METH_NOARGS, "x.copy() -> FieldPath -- copy this Unreal field path" },
+		{ "is_valid", PyCFunctionCast(&FMethods::IsValid), METH_NOARGS, "x.is_valid() -> bool -- whether this Unreal field path refers to an existing Unreal field" },
+		{ nullptr, nullptr, 0, nullptr }
+	};
+
+	PyTypeObject PyType = InitializePyWrapperBasicType<FPyWrapperFieldPath>("FieldPath", "Type for all Unreal exposed FieldPath instances");
+
+	PyType.tp_init = (initproc)&FFuncs::Init;
+	PyType.tp_str = (reprfunc)&FFuncs::Str;
+	PyType.tp_repr = (reprfunc)&FFuncs::Repr;
+	PyType.tp_richcompare = (richcmpfunc)&FFuncs::RichCmp;
+	PyType.tp_hash = (hashfunc)&FFuncs::Hash;
+
+	PyType.tp_methods = PyMethods;
+
+	return PyType;
+}
+
+PyTypeObject PyWrapperFieldPathType = InitializePyWrapperFieldPathType();
+
+#endif	// WITH_PYTHON
diff --git a/Engine/Plugins/Experimental/PythonScriptPlugin/Source/PythonScriptPlugin/Private/PyWrapperFieldPath.h b/Engine/Plugins/Experimental/PythonScriptPlugin/Source/PythonScriptPlugin/Private/PyWrapperFieldPath.h
new file mode 100644
--- /dev/null
+++ b/Engine/Plugins/Experimental/PythonScriptPlugin/Source/PythonScriptPlugin/Private/PyWrapperFieldPath.h
@@ -0,0 +1,35 @@
+// Copyright Epic Games, Inc. All Rights Reserved.
+
+#pragma once
+
+#include "PyWrapperBasic.h"
+
+#if WITH_PYTHON
+
+/** Python type for FPyWrapperFieldPath */
+extern PyTypeObject PyWrapperFieldPathType;
+
+/** Initialize the PyWrapperFieldPath types and add them to the given Python module */
+void InitializePyWrapperFieldPath(PyGenUtil::FNativePythonModule& ModuleInfo);
+
+/** Type for all Unreal exposed FPyWrapperFieldPath instances */
+struct FPyWrapperFieldPath : public TPyWrapperBasic<FFieldPath, FPyWrapperFieldPath>
+{
+	typedef TPyWrapperBasic<FFieldPath, FPyWrapperFieldPath> Super;
+
+	/** Cast the given Python object to this wrapped type (returns a new reference) */
+	static FPyWrapperFieldPath* CastPyObject(PyObject* InPyObject, FPyConversionResult* OutCastResult = nullptr);
+
+	/** Cast the given Python object to this wrapped type, or attempt to convert the type into a new wrapped instance (returns a new reference) */
+	static FPyWrapperFieldPath* CastPyObject(PyObject* InPyObject, PyTypeObject* InType, FPyConversionResult* OutCastResult = nullptr);
+
+	/** Initialize the value of this wrapper instance (internal) */
+	static void InitValue(FPyWrapperFieldPath* InSelf, const FFieldPath& InValue);
+
+	/** Deinitialize this wrapper instance (called via Init and Free to restore the instance to its New state) */
+	static void DeinitValue(FPyWrapperFieldPath* InSelf);
+};
+
+typedef TPyPtr<FPyWrapperFieldPath> FPyWrapperFieldPathPtr;
+
+#endif	// WITH_PYTHON
\ No newline at end of file
diff --git a/Engine/Plugins/Experimental/PythonScriptPlugin/Source/PythonScriptPlugin/Private/PyWrapperTypeRegistry.cpp b/Engine/Plugins/Experimental/PythonScriptPlugin/Source/PythonScriptPlugin/Private/PyWrapperTypeRegistry.cpp
--- a/Engine/Plugins/Experimental/PythonScriptPlugin/Source/PythonScriptPlugin/Private/PyWrapperTypeRegistry.cpp
+++ b/Engine/Plugins/Experimental/PythonScriptPlugin/Source/PythonScriptPlugin/Private/PyWrapperTypeRegistry.cpp
@@ -12,6 +12,7 @@
 #include "PyWrapperFixedArray.h"
 #include "PyWrapperSet.h"
 #include "PyWrapperMap.h"
+#include "PyWrapperFieldPath.h"
 #include "PyConversion.h"
 #include "PyGIL.h"
 #include "PyCore.h"
@@ -436,6 +437,24 @@ FPyWrapperMap* FPyWrapperMapFactory::CreateInstance(void* InUnrealInstance, cons
 	}, InConversionMethod == EPyConversionMethod::Copy || InConversionMethod == EPyConversionMethod::Steal);
 }
 
+FPyWrapperFieldPathFactory& FPyWrapperFieldPathFactory::Get()
+{
+	static FPyWrapperFieldPathFactory Instance;
+	return Instance;
+}
+
+FPyWrapperFieldPath* FPyWrapperFieldPathFactory::FindInstance(FFieldPath InUnrealInstance) const
+{
+	return FindInstanceInternal(InUnrealInstance, &PyWrapperFieldPathType);
+}
+
+FPyWrapperFieldPath* FPyWrapperFieldPathFactory::CreateInstance(FFieldPath InUnrealInstance)
+{
+	return CreateInstanceInternal(InUnrealInstance, &PyWrapperFieldPathType, [InUnrealInstance](FPyWrapperFieldPath* InSelf)
+	{
+		return FPyWrapperFieldPath::Init(InSelf, InUnrealInstance);
+	});
+}
 
 FPyWrapperTypeReinstancer& FPyWrapperTypeReinstancer::Get()
 {
@@ -2313,7 +2332,7 @@ void FPyWrapperTypeRegistry::GenerateStubCodeForWrappedTypes(const EPyOnlineDocs
 								if (!PythonFileLine.Contains(TEXT("unreal"), ESearchCase::CaseSensitive))
 								{
 									bExportedImports = true;
-									PythonScript.WriteLine(PythonFileLine);
+									PythonScript.WriteLine(PythonFileLine.TrimStart()); // Trim leading spaces, import in .py file might be indented (lazy loaded in a if).
 								}
 								continue;
 							}
diff --git a/Engine/Plugins/Experimental/PythonScriptPlugin/Source/PythonScriptPlugin/Private/PyWrapperTypeRegistry.h b/Engine/Plugins/Experimental/PythonScriptPlugin/Source/PythonScriptPlugin/Private/PyWrapperTypeRegistry.h
--- a/Engine/Plugins/Experimental/PythonScriptPlugin/Source/PythonScriptPlugin/Private/PyWrapperTypeRegistry.h
+++ b/Engine/Plugins/Experimental/PythonScriptPlugin/Source/PythonScriptPlugin/Private/PyWrapperTypeRegistry.h
@@ -26,6 +26,7 @@ struct FPyWrapperFixedArray;
 struct FPyWrapperSet;
 struct FPyWrapperMap;
 struct FPyWrapperStructMetaData;
+struct FPyWrapperFieldPath;
 
 class FPyWrapperOwnerContext;
 class UPythonGeneratedClass;
@@ -299,6 +300,21 @@ public:
 	FPyWrapperMap* CreateInstance(void* InUnrealInstance, const FMapProperty* InProp, const FPyWrapperOwnerContext& InOwnerContext, const EPyConversionMethod InConversionMethod);
 };
 
+/** Factory for wrapped field type instances */
+class FPyWrapperFieldPathFactory : public TPyWrapperTypeFactory<FFieldPath, FPyWrapperFieldPath>
+{
+public:
+	/** Access the singleton instance */
+	static FPyWrapperFieldPathFactory& Get();
+
+	/** Find the wrapped Python instance associated with the given Unreal instance (if any, returns borrowed reference) */
+	FPyWrapperFieldPath* FindInstance(FFieldPath InUnrealInstance) const;
+
+	/** Find the wrapped Python instance associated with the given Unreal instance, or create one if needed (returns new reference) */
+	FPyWrapperFieldPath* CreateInstance(FFieldPath InUnrealInstance);
+};
+
+
 /** Singleton instance that handles re-instancing Python types */
 class FPyWrapperTypeReinstancer
 {
