Changed a Python glue generation error into a warning when a make or a break function signature doesn't comply with Python requirements. Some licensees have sophisticated make/break functions that cannot be easily supported by Python. Downgrading from an error to a warning, with a message saying that the default make/break will be used instead to make/break the Python object, will prevent licensees from hitting a UE_DEBUG_BREAK every time the debugger is attached.

diff --git a/Engine/Plugins/Experimental/PythonScriptPlugin/Source/PythonScriptPlugin/Private/PyWrapperTypeRegistry.cpp b/Engine/Plugins/Experimental/PythonScriptPlugin/Source/PythonScriptPlugin/Private/PyWrapperTypeRegistry.cpp
--- a/Engine/Plugins/Experimental/PythonScriptPlugin/Source/PythonScriptPlugin/Private/PyWrapperTypeRegistry.cpp
+++ b/Engine/Plugins/Experimental/PythonScriptPlugin/Source/PythonScriptPlugin/Private/PyWrapperTypeRegistry.cpp
@@ -81,7 +81,7 @@ void SetMakeFunction(FPyWrapperStructMetaData& MetaData, UFunction* MakeFunc)
 
 	if (!bHasValidReturn)
 	{
-		REPORT_PYTHON_GENERATION_ISSUE(Error, TEXT("Struct '%s' is marked as 'HasNativeMake' but the function '%s' does not return the struct type."), *MetaData.Struct->GetName(), *MetaData.MakeFunc.Func->GetPathName());
+		REPORT_PYTHON_GENERATION_ISSUE(Warning, TEXT("Struct '%s' is marked as 'HasNativeMake' but the function '%s' does not return the struct type. Python will use the generic make function to create a Python object of this type."), *MetaData.Struct->GetName(), *MetaData.MakeFunc.Func->GetPathName());
 		MetaData.MakeFunc.SetFunction(nullptr);
 		return;
 	}
@@ -109,7 +109,7 @@ void SetBreakFunction(FPyWrapperStructMetaData& MetaData, UFunction* BreakFunc)
 
 	if (!bHasValidInput)
 	{
-		REPORT_PYTHON_GENERATION_ISSUE(Error, TEXT("Struct '%s' is marked as 'HasNativeBreak' but the function '%s' does not have the struct type as its only input argument."), *MetaData.Struct->GetName(), *MetaData.BreakFunc.Func->GetPathName());
+		REPORT_PYTHON_GENERATION_ISSUE(Warning, TEXT("Struct '%s' is marked as 'HasNativeBreak' but the function '%s' does not have the struct type as its only input argument. Python will use the generic break function to convert a Python object of this type into a tuple."), *MetaData.Struct->GetName(), *MetaData.BreakFunc.Func->GetPathName());
 		MetaData.BreakFunc.SetFunction(nullptr);
 	}
 };
@@ -2645,12 +2645,12 @@ void FPyWrapperTypeRegistry::GenerateStubCodeForWrappedType(PyTypeObject* PyType
 			}
 			else if (!FirstDefaultedParamName.IsEmpty())
 			{
-				// Some UFUNCTION has their default values declared as 'meta' invisible to C++ compile, but Python generated glue will use that. Sometime, the next parameter(s) are
+				// Some UFUNCTION has their default values declared as 'meta' invisible to C++ compiler, but Python generated glue will use that. Sometime, the next parameter(s) are
 				// not defaulted. This creates invalid Python declarations. In stub, this is legal to use "..." to indicate a default param value without specifying the exact value.
 				// That's a workaround for faulty method declarations and the C++ code should be fixed.
 				MethodArgsStr += TEXT("=...");
 
-				// Left as 'Verbose' until we fix known cases. Once fixed, we should turn this as a 'Warning'.
+				// Left as 'Verbose' until we fix known cases. Once fixed, we should turn this into a 'Warning'.
 				REPORT_PYTHON_GENERATION_ISSUE(Verbose, TEXT("The function '%s' is missing default values for param '%s'. All params after the first defaulted param '%s' (likely from meta data) should be defaulted. The Python stub defaulted missing default values."), 
 					*InTypeMethod.MethodFunc.Func->GetPathName(), *MethodParam.ParamProp->GetName(), *FirstDefaultedParamName);
 			}
