Fixed static analysis error in PyWrapperTypeRegistry.cpp

diff --git a/Engine/Plugins/Experimental/PythonScriptPlugin/Source/PythonScriptPlugin/Private/PyWrapperTypeRegistry.cpp b/Engine/Plugins/Experimental/PythonScriptPlugin/Source/PythonScriptPlugin/Private/PyWrapperTypeRegistry.cpp
--- a/Engine/Plugins/Experimental/PythonScriptPlugin/Source/PythonScriptPlugin/Private/PyWrapperTypeRegistry.cpp
+++ b/Engine/Plugins/Experimental/PythonScriptPlugin/Source/PythonScriptPlugin/Private/PyWrapperTypeRegistry.cpp
@@ -2652,7 +2652,9 @@ void FPyWrapperTypeRegistry::GenerateStubCodeForWrappedType(PyTypeObject* PyType
 
 				// Left as 'Verbose' until we fix known cases. Once fixed, we should turn this into a 'Warning'.
 				REPORT_PYTHON_GENERATION_ISSUE(Verbose, TEXT("The function '%s' is missing default values for param '%s'. All params after the first defaulted param '%s' (likely from meta data) should be defaulted. The Python stub defaulted missing default values."), 
-					*InTypeMethod.MethodFunc.Func->GetPathName(), *MethodParam.ParamProp->GetName(), *FirstDefaultedParamName);
+					InTypeMethod.MethodFunc.Func != nullptr ? *InTypeMethod.MethodFunc.Func->GetPathName() : TEXT("Unknown"),
+					MethodParam.ParamProp != nullptr ? *MethodParam.ParamProp->GetName() : TEXT("Unknown"),
+					*FirstDefaultedParamName);
 			}
 		}
 
