Fixed missing function declaration of Nativize(FieldPath) in PyConversion.h

diff --git a/Engine/Plugins/Experimental/PythonScriptPlugin/Source/PythonScriptPlugin/Private/PyConversion.h b/Engine/Plugins/Experimental/PythonScriptPlugin/Source/PythonScriptPlugin/Private/PyConversion.h
--- a/Engine/Plugins/Experimental/PythonScriptPlugin/Source/PythonScriptPlugin/Private/PyConversion.h
+++ b/Engine/Plugins/Experimental/PythonScriptPlugin/Source/PythonScriptPlugin/Private/PyConversion.h
@@ -77,6 +77,10 @@ namespace PyConversion
 	FPyConversionResult Nativize(PyObject* PyObj, FText& OutVal, const ESetErrorState SetErrorState = ESetErrorState::Yes);
 	FPyConversionResult Pythonize(const FText& Val, PyObject*& OutPyObj, const ESetErrorState SetErrorState = ESetErrorState::Yes);
 
+	/** FFieldPath overload */
+	FPyConversionResult Nativize(PyObject* PyObj, FFieldPath& OutVal, const ESetErrorState SetErrorState = ESetErrorState::Yes);
+	FPyConversionResult Pythonize(const FFieldPath& Val, PyObject*& OutPyObj, const ESetErrorState SetErrorState = ESetErrorState::Yes);
+
 	/** void* overload */
 	FPyConversionResult Nativize(PyObject* PyObj, void*& OutVal, const ESetErrorState SetErrorState = ESetErrorState::Yes);
 	FPyConversionResult Pythonize(void* Val, PyObject*& OutPyObj, const ESetErrorState SetErrorState = ESetErrorState::Yes);
