Fixed Python calling the wrong delegate if the same delegate name/type was declared in two different UObject.
Added a test in test_wrapper_types.py to validate before and after the fix.

diff --git a/Engine/Plugins/Experimental/PythonScriptPlugin/Source/PythonScriptPlugin/Private/PyGenUtil.cpp b/Engine/Plugins/Experimental/PythonScriptPlugin/Source/PythonScriptPlugin/Private/PyGenUtil.cpp
--- a/Engine/Plugins/Experimental/PythonScriptPlugin/Source/PythonScriptPlugin/Private/PyGenUtil.cpp
+++ b/Engine/Plugins/Experimental/PythonScriptPlugin/Source/PythonScriptPlugin/Private/PyGenUtil.cpp
@@ -13,10 +13,12 @@
 #include "Misc/FileHelper.h"
 #include "Misc/PackageName.h"
 #include "Misc/OutputDeviceNull.h"
+#include "Misc/StringBuilder.h"
 #include "UObject/Class.h"
 #include "UObject/Stack.h"
 #include "UObject/Package.h"
 #include "UObject/EnumProperty.h"
+#include "UObject/NameTypes.h"
 #include "UObject/TextProperty.h"
 #include "UObject/CoreRedirects.h"
 #include "UObject/PropertyPortFlags.h"
@@ -2498,6 +2500,16 @@ FName GetTypeRegistryName(const UEnum* InEnum)
 
 FName GetTypeRegistryName(const UFunction* InDelegateSignature)
 {
+	// Two different UObject-based classes can each declare in their body a delegate with the same type name, but possibly
+	// with different parameters. While they will share the same FName, they aren't be the same type. Make the type registry
+	// name unique to each delegate.
+	if (UObject* Outer = InDelegateSignature->GetOuter())
+	{
+		FNameBuilder ConcatName(Outer->GetFName());
+		ConcatName += "_";
+		InDelegateSignature->GetFName().AppendString(ConcatName);
+		return FName(*ConcatName);
+	}
 	return InDelegateSignature->GetFName();
 }
 
diff --git a/Engine/Plugins/Experimental/PythonScriptPlugin/Source/PythonScriptPlugin/Private/PyTest.h b/Engine/Plugins/Experimental/PythonScriptPlugin/Source/PythonScriptPlugin/Private/PyTest.h
--- a/Engine/Plugins/Experimental/PythonScriptPlugin/Source/PythonScriptPlugin/Private/PyTest.h
+++ b/Engine/Plugins/Experimental/PythonScriptPlugin/Source/PythonScriptPlugin/Private/PyTest.h
@@ -288,3 +288,35 @@ class UPyTestObjectLibrary : public UBlueprintFunctionLibrary
 	UFUNCTION(BlueprintPure, Category = "Python|Internal", meta=(ScriptConstant="OtherConstantValue", ScriptConstantHost="PyTestObject"))
 	static int32 GetOtherConstantValue();
 };
+
+/**
+ * This class along with UPyTestVectorDelegate verify that 2 UObjects with the same delegate name/type, do not collide.
+ */
+UCLASS(Blueprintable)
+class UPyTestStructDelegate : public UObject
+{
+	GENERATED_BODY()
+
+public:
+	DECLARE_DYNAMIC_MULTICAST_DELEGATE_OneParam(FOnNameCollisionDelegate, const FPyTestStruct&, PyStruct);
+
+	/** Called when a new item is selected in the combobox. */
+	UPROPERTY(BlueprintAssignable, Category = "Python|Internal")
+	FOnNameCollisionDelegate OnNameCollisionDelegate; // Same prop name and type name as UPyTestVectorDelegate::FOnNameCollisionDelegate, but different params.
+};
+
+/**
+ * This class along with UPyTestStructDelegate verify that 2 UObjects with the same delegate name/type, do not collide.
+ */
+UCLASS(Blueprintable)
+class UPyTestVectorDelegate : public UObject
+{
+	GENERATED_BODY()
+
+public:
+	DECLARE_DYNAMIC_MULTICAST_DELEGATE_OneParam(FOnNameCollisionDelegate, const FVector2D&, Vec);
+
+	/** Called when a new item is selected in the combobox. */
+	UPROPERTY(BlueprintAssignable, Category = "Python|Internal")
+	FOnNameCollisionDelegate OnNameCollisionDelegate; // Same prop name and type name as UPyTestStructDelegate::FOnNameCollisionDelegate, but different params.
+};
diff --git a/Engine/Plugins/Experimental/PythonScriptPlugin/Source/PythonScriptPlugin/Private/PyWrapperTypeRegistry.cpp b/Engine/Plugins/Experimental/PythonScriptPlugin/Source/PythonScriptPlugin/Private/PyWrapperTypeRegistry.cpp
--- a/Engine/Plugins/Experimental/PythonScriptPlugin/Source/PythonScriptPlugin/Private/PyWrapperTypeRegistry.cpp
+++ b/Engine/Plugins/Experimental/PythonScriptPlugin/Source/PythonScriptPlugin/Private/PyWrapperTypeRegistry.cpp
@@ -22,6 +22,7 @@
 #include "ProfilingDebugging/CpuProfilerTrace.h"
 #include "Misc/Paths.h"
 #include "Misc/FileHelper.h"
+#include "Misc/StringBuilder.h"
 #include "SourceCodeNavigation.h"
 #include "UObject/UnrealType.h"
 #include "UObject/UObjectHash.h"
@@ -29,6 +30,7 @@
 #include "UObject/Package.h"
 #include "UObject/EnumProperty.h"
 #include "UObject/StructOnScope.h"
+#include "UObject/NameTypes.h"
 #include "Misc/CoreDelegates.h"
 #if WITH_EDITOR
 #include "Kismet2/ReloadUtilities.h"
@@ -2086,10 +2088,28 @@ PyTypeObject* FPyWrapperTypeRegistry::GenerateWrappedDelegateType(const UFunctio
 
 	GeneratedWrappedType->PyType.tp_flags = Py_TPFLAGS_DEFAULT;
 
+	// Two different UObject-based classes can each declare in their body a delegate with the same type name, but possibly
+	// with different parameters. While they will share the same 'DelegateBaseTypename', they aren't necessarily the same type. Make
+	// the name of the 'PythonCallableForDelegateClass' unique because NewObject<> will find pre-existing object with the same name and
+	// return the same address, then the code below will basically erase the previous delegate type with the new delegate type, possibly
+	// confusing the parameter types if they aren't the same.
+	FNameBuilder PythonCallableForDelegateObjectName;
+	if (UObject* Outer = InDelegateSignature->GetOuter())
+	{
+		Outer->GetFName().AppendString(PythonCallableForDelegateObjectName);
+		PythonCallableForDelegateObjectName += "_";
+		PythonCallableForDelegateObjectName += DelegateBaseTypename;
+	}
+	else
+	{
+		PythonCallableForDelegateObjectName += DelegateBaseTypename;
+	}
+	PythonCallableForDelegateObjectName += TEXT("__PythonCallable");
+
 	// Generate the proxy class needed to wrap Python callables in Unreal delegates
 	UClass* PythonCallableForDelegateClass = nullptr;
 	{
-		PythonCallableForDelegateClass = NewObject<UClass>(GetPythonTypeContainer(), *FString::Printf(TEXT("%s__PythonCallable"), *DelegateBaseTypename), RF_Public | RF_Standalone | RF_Transient);
+		PythonCallableForDelegateClass = NewObject<UClass>(GetPythonTypeContainer(), *PythonCallableForDelegateObjectName, RF_Public | RF_Standalone | RF_Transient);
 		UFunction* PythonCallableForDelegateFunc = nullptr;
 		{
 			FObjectDuplicationParameters FuncDuplicationParams(const_cast<UFunction*>(InDelegateSignature), PythonCallableForDelegateClass);
