Implemented Python Type Hinting when generating the Python stub.
  - The 'Off' mode turns type hinting off and produce the Python stub as before.
  - The 'Auto-Completion' mode aims to provide types in Python IDE auto-completion while leaving out some information/complexity that clutter the function definition.
  - The 'Type Checker' mode aims to provide the closest to the truth (with some limitations) of the supported API types. Type coercion makes it harder to see the exact types.

Known limitations:
  - We don't know when a reflected UObject can legally be None. In 'Type Checker' mode, the UObject are always marked as typing.Optional[] everywhere to denotate that they can be None. Obviously, some methods never returns None and some method will fails if an Object is None. It might be annoying to users in some case.

diff --git a/Engine/Plugins/Experimental/PythonScriptPlugin/Source/PythonScriptPlugin/Private/PyCore.cpp b/Engine/Plugins/Experimental/PythonScriptPlugin/Source/PythonScriptPlugin/Private/PyCore.cpp
--- a/Engine/Plugins/Experimental/PythonScriptPlugin/Source/PythonScriptPlugin/Private/PyCore.cpp
+++ b/Engine/Plugins/Experimental/PythonScriptPlugin/Source/PythonScriptPlugin/Private/PyCore.cpp
@@ -649,11 +649,11 @@ PyTypeObject InitializePyScopedSlowTaskType()
 	};
 
 	static PyMethodDef PyMethods[] = {
-		{ "__enter__", PyCFunctionCast(&FMethods::EnterScope), METH_NOARGS, "x.__enter__() -> self -- begin this slow task" },
-		{ "__exit__", PyCFunctionCast(&FMethods::ExitScope), METH_VARARGS | METH_KEYWORDS, "x.__exit__(type, value, traceback) -> None -- end this slow task" },
-		{ "make_dialog_delayed", PyCFunctionCast(&FMethods::MakeDialogDelayed), METH_VARARGS | METH_KEYWORDS, "x.make_dialog_delayed(delay, can_cancel=False, allow_in_pie=False) -> None -- creates a new dialog for this slow task after the given time threshold. If the task completes before this time, no dialog will be shown" },
-		{ "make_dialog", PyCFunctionCast(&FMethods::MakeDialog), METH_VARARGS | METH_KEYWORDS, "x.make_dialog(can_cancel=False, allow_in_pie=False) -> None -- creates a new dialog for this slow task, if there is currently not one open" },
-		{ "enter_progress_frame", PyCFunctionCast(&FMethods::EnterProgressFrame), METH_VARARGS | METH_KEYWORDS, "x.enter_progress_frame(work=1.0, desc=Text()) -> None -- indicate that we are to enter a frame that will take up the specified amount of work (completes any previous frames)" },
+		{ "__enter__", PyCFunctionCast(&FMethods::EnterScope), METH_NOARGS, "__enter__(self) -> ScopedSlowTask -- begin this slow task" },
+		{ "__exit__", PyCFunctionCast(&FMethods::ExitScope), METH_VARARGS | METH_KEYWORDS, "__exit__(self, type: Optional[Type[BaseException]], value: Optional[BaseException], traceback: Optional[TracebackType]) -> bool -- end this slow task" },
+		{ "make_dialog_delayed", PyCFunctionCast(&FMethods::MakeDialogDelayed), METH_VARARGS | METH_KEYWORDS, "make_dialog_delayed(self, delay: float, can_cancel: bool = False, allow_in_pie: bool = False) -> None -- creates a new dialog for this slow task after the given time delay (in seconds). If the task completes before this time, no dialog will be shown" },
+		{ "make_dialog", PyCFunctionCast(&FMethods::MakeDialog), METH_VARARGS | METH_KEYWORDS, "make_dialog(self, can_cancel: bool = False, allow_in_pie: bool = False) -> None -- creates a new dialog for this slow task, if there is currently not one open" },
+		{ "enter_progress_frame", PyCFunctionCast(&FMethods::EnterProgressFrame), METH_VARARGS | METH_KEYWORDS, "enter_progress_frame(self, work: float = 1.0, desc: Union[Text, str] = \"\") -> None -- indicate that we are to enter a frame that will take up the specified amount of work (completes any previous frames)" },
 		{ "should_cancel", PyCFunctionCast(&FMethods::ShouldCancel), METH_NOARGS, "x.should_cancel() -> bool -- True if the user has requested that the slow task be canceled" },
 		{ nullptr, nullptr, 0, nullptr }
 	};
@@ -1247,7 +1247,8 @@ PyObject* NewObject(PyObject* InSelf, PyObject* InArgs, PyObject* InKwds)
 			return nullptr;
 		}
 
-		if (PyOuterObj && !PyConversion::Nativize(PyOuterObj, ObjectOuter))
+		// If PyOuterObj is None, Nativize() will convert 'PyOuterObj=None' into 'nullptr'. That will overwrite the desired default value 'transient package'.
+		if (PyOuterObj && PyOuterObj != Py_None && !PyConversion::Nativize(PyOuterObj, ObjectOuter))
 		{
 			PyUtil::SetPythonError(PyExc_TypeError, TEXT("new_object"), *FString::Printf(TEXT("Failed to convert 'outer' (%s) to 'Object'"), *PyUtil::GetFriendlyTypename(PyOuterObj)));
 			return nullptr;
@@ -1799,36 +1800,39 @@ PyObject* UnregisterPythonShutdownCallback(PyObject* InSelf, PyObject* InArgs)
 	Py_RETURN_NONE;
 }
 
+// NOTE: Several functions like 'new_object' are marked as returning 'Any' because the exact return type cannot be deduced from the Class type. We could mark them as returning base 'Object'
+//       type, but that would force user to cast down to please type checker while the returned type is already the exact requested type. For example, the instruction
+//       'o: PyTestObject = unreal.new_object(unreal.PyTestObject.static_class())' must pass type checker and if new_object() was to return Object, it would not be automatically cast down.
 PyMethodDef PyCoreMethods[] = {
-	{ "log", PyCFunctionCast(&Log), METH_VARARGS, "x.log(str) -> None -- log the given argument as information in the LogPython category" },
-	{ "log_warning", PyCFunctionCast(&LogWarning), METH_VARARGS, "x.log_warning(str) -> None -- log the given argument as a warning in the LogPython category" },
-	{ "log_error", PyCFunctionCast(&LogError), METH_VARARGS, "x.log_error(str) -> None -- log the given argument as an error in the LogPython category" },
-	{ "log_flush", PyCFunctionCast(&LogFlush), METH_NOARGS, "x.log_flush() -> None -- flush the log to disk" },
-	{ "reload", PyCFunctionCast(&Reload), METH_VARARGS, "x.reload(str) -> None -- reload the given Unreal Python module" },
-	{ "load_module", PyCFunctionCast(&LoadModule), METH_VARARGS, "x.load_module(str) -> None -- load the given Unreal module and generate any Python code for its reflected types" },
-	{ "new_object", PyCFunctionCast(&NewObject), METH_VARARGS | METH_KEYWORDS, "x.new_object(type, outer=Transient, name=Default, base_type=Object) -> Object -- create an Unreal object of the given class (and optional outer and name), optionally validating its type" },
-	{ "find_object", PyCFunctionCast(&FindObject), METH_VARARGS | METH_KEYWORDS, "x.find_object(outer, name, type=Object, follow_redirectors=True) -> Object -- find an already loaded Unreal object with the given outer and name, optionally validating its type" },
-	{ "load_object", PyCFunctionCast(&LoadObject), METH_VARARGS | METH_KEYWORDS, "x.load_object(outer, name, type=Object, follow_redirectors=True) -> Object -- load an Unreal object with the given outer and name, optionally validating its type" },
-	{ "load_class", PyCFunctionCast(&LoadClass), METH_VARARGS | METH_KEYWORDS, "x.load_class(outer, name, type=Object) -> Class -- load an Unreal class with the given outer and name, optionally validating its base type" },
-	{ "find_asset", PyCFunctionCast(&FindAsset), METH_VARARGS | METH_KEYWORDS, "x.find_asset(name, type=Object, follow_redirectors=True) -> Object -- find an already loaded Unreal asset with the given name, optionally validating its type" },
-	{ "load_asset", PyCFunctionCast(&LoadAsset), METH_VARARGS | METH_KEYWORDS, "x.load_asset(name, type=Object, follow_redirectors=True) -> Object -- load an Unreal asset with the given name, optionally validating its type" },
-	{ "find_package", PyCFunctionCast(&FindPackage), METH_VARARGS, "x.find_package(name) -> Package -- find an already loaded Unreal package with the given name" },
-	{ "load_package", PyCFunctionCast(&LoadPackage), METH_VARARGS, "x.load_package(name) -> Package -- load an Unreal package with the given name" },
-	{ "get_default_object", PyCFunctionCast(&GetDefaultObject), METH_VARARGS, "x.get_default_object(type) -> Object -- get the Unreal class default object (CDO) of the given type" },
-	{ "purge_object_references", PyCFunctionCast(&PurgeObjectReferences), METH_VARARGS | METH_KEYWORDS, "x.purge_object_references(obj, include_inners=True) -> None -- purge all references to the given Unreal object from any living Python objects" },
-	{ "generate_class", PyCFunctionCast(&GenerateClass), METH_VARARGS, "x.generate_class(type) -> None -- generate an Unreal class for the given Python type" },
-	{ "generate_struct", PyCFunctionCast(&GenerateStruct), METH_VARARGS, "x.generate_struct(type) -> None -- generate an Unreal struct for the given Python type" },
-	{ "generate_enum", PyCFunctionCast(&GenerateEnum), METH_VARARGS, "x.generate_enum(type) -> None -- generate an Unreal enum for the given Python type" },
-	{ "flush_generated_type_reinstancing", PyCFunctionCast(&FlushGeneratedTypeReinstancing), METH_NOARGS, "x.flush_generated_type_reinstancing() -> None -- flush any pending reinstancing requests for Python generated types" },
-	{ "get_type_from_class", PyCFunctionCast(&GetTypeFromClass), METH_VARARGS, "x.get_type_from_class(class) -> type -- get the best matching Python type for the given Unreal class" },
-	{ "get_type_from_struct", PyCFunctionCast(&GetTypeFromStruct), METH_VARARGS, "x.get_type_from_struct(struct) -> type -- get the best matching Python type for the given Unreal struct" },
-	{ "get_type_from_enum", PyCFunctionCast(&GetTypeFromEnum), METH_VARARGS, "x.get_type_from_enum(enum) -> type -- get the best matching Python type for the given Unreal enum" },
-	{ "register_python_shutdown_callback", PyCFunctionCast(&RegisterPythonShutdownCallback), METH_VARARGS, "x.register_python_shutdown_callback(callable) -> _DelegateHandle -- register the given callable (with no input arguments) as a callback to execute immediately before Python shutdown"},
-	{ "unregister_python_shutdown_callback", PyCFunctionCast(&UnregisterPythonShutdownCallback), METH_VARARGS, "x.unregister_python_shutdown_callback(handle) -> None -- unregister the given handle from a previous call to register_python_shutdown_callback"},
-	{ "NSLOCTEXT", PyCFunctionCast(&CreateLocalizedText), METH_VARARGS, "x.NSLOCTEXT(ns, key, source) -> Text -- create a localized Text from the given namespace, key, and source string" },
-	{ "LOCTABLE", PyCFunctionCast(&CreateLocalizedTextFromStringTable), METH_VARARGS, "x.LOCTABLE(id, key) -> Text -- get a localized Text from the given string table id and key" },
-	{ "is_editor", PyCFunctionCast(&IsEditor), METH_NOARGS, "x.is_editor() -> Bool -- tells if the editor is running or not" },
-	{ "get_interpreter_executable_path", PyCFunctionCast(&GetInterpreterExecutablePath), METH_NOARGS, "x.get_interpreter_executable_path() -> str -- get the path to the Python interpreter executable of the Python SDK this plugin was compiled against" },
+	{ "log", PyCFunctionCast(&Log), METH_VARARGS, "log(arg: Any) -> None -- log the given argument as information in the LogPython category" },
+	{ "log_warning", PyCFunctionCast(&LogWarning), METH_VARARGS, "log_warning(arg: Any) -> None -- log the given argument as a warning in the LogPython category" },
+	{ "log_error", PyCFunctionCast(&LogError), METH_VARARGS, "log_error(arg: Any) -> None -- log the given argument as an error in the LogPython category" },
+	{ "log_flush", PyCFunctionCast(&LogFlush), METH_NOARGS, "log_flush() -> None -- flush the log to disk." },
+	{ "reload", PyCFunctionCast(&Reload), METH_VARARGS, "reload(module: str) -> None -- reload the given Unreal Python module." },
+	{ "load_module", PyCFunctionCast(&LoadModule), METH_VARARGS, "load_module(module: str) -> None -- load the given Unreal module and generate any Python code for its reflected types" },
+	{ "new_object", PyCFunctionCast(&NewObject), METH_VARARGS | METH_KEYWORDS, "new_object(type: Union[Class, type], outer: Optional[Object]=None, name: Union[Name, str]=\"\", base_type: Optional[Object]=None) -> Any -- create an Unreal object of the given class (and optional outer and name), optionally validating its type" },
+	{ "find_object", PyCFunctionCast(&FindObject), METH_VARARGS | METH_KEYWORDS, "find_object(outer: Optional[Object], name: str, type: Union[Class, type]=Object.static_class(), follow_redirectors: bool=True) -> Any -- find an already loaded Unreal object with the given outer and name, optionally validating its type" },
+	{ "load_object", PyCFunctionCast(&LoadObject), METH_VARARGS | METH_KEYWORDS, "load_object(outer: Optional[Object], name: str, type: Union[Class, type]=Object.static_class(), follow_redirectors: bool=True) -> Any -- load an Unreal object with the given outer and name, optionally validating its type" },
+	{ "load_class", PyCFunctionCast(&LoadClass), METH_VARARGS | METH_KEYWORDS, "load_class(outer: Optional[Object], name: str, type: Union[Class, type]=Object.static_class()) -> Optional[Class] -- load an Unreal class with the given outer and name, optionally validating its base type" },
+	{ "find_asset", PyCFunctionCast(&FindAsset), METH_VARARGS | METH_KEYWORDS, "find_asset(name: str, type: Union[Class, type]=Object.static_class(), follow_redirectors: bool = True) -> Any -- find an already loaded Unreal asset with the given name, optionally validating its type" },
+	{ "load_asset", PyCFunctionCast(&LoadAsset), METH_VARARGS | METH_KEYWORDS, "load_asset(name: str, type: Union[Class, type]=Object.static_class(), follow_redirectors: bool = True) -> Any -- load an Unreal asset with the given name, optionally validating its type" },
+	{ "find_package", PyCFunctionCast(&FindPackage), METH_VARARGS, "find_package(name: str) -> Optional[Package] -- find an already loaded Unreal package with the given name" },
+	{ "load_package", PyCFunctionCast(&LoadPackage), METH_VARARGS, "load_package(name: str) -> Optional[Package] -- load an Unreal package with the given name" },
+	{ "get_default_object", PyCFunctionCast(&GetDefaultObject), METH_VARARGS, "get_default_object(type: Union[Class, type]) -> Any -- get the Unreal class default object (CDO) of the given type" },
+	{ "purge_object_references", PyCFunctionCast(&PurgeObjectReferences), METH_VARARGS | METH_KEYWORDS, "purge_object_references(obj: Object, include_inners: bool = True) -> None -- purge all references to the given Unreal object from any living Python objects" },
+	{ "generate_class", PyCFunctionCast(&GenerateClass), METH_VARARGS, "generate_class(class_type: type) -> None -- generate an Unreal class for the given Python type" },
+	{ "generate_struct", PyCFunctionCast(&GenerateStruct), METH_VARARGS, "generate_struct(struct_type: type) -> None -- generate an Unreal struct for the given Python type" },
+	{ "generate_enum", PyCFunctionCast(&GenerateEnum), METH_VARARGS, "generate_enum(enum_type: type) -> None -- generate an Unreal enum for the given Python type" },
+	{ "flush_generated_type_reinstancing", PyCFunctionCast(&FlushGeneratedTypeReinstancing), METH_NOARGS, "flush_generated_type_reinstancing() -> None -- flush any pending reinstancing requests for Python generated types" },
+	{ "get_type_from_class", PyCFunctionCast(&GetTypeFromClass), METH_VARARGS, "get_type_from_class(class_: Class) -> type -- get the best matching Python type for the given Unreal class" },
+	{ "get_type_from_struct", PyCFunctionCast(&GetTypeFromStruct), METH_VARARGS, "get_type_from_struct(struct: Struct) -> type -- get the best matching Python type for the given Unreal struct" },
+	{ "get_type_from_enum", PyCFunctionCast(&GetTypeFromEnum), METH_VARARGS, "get_type_from_enum(enum: Enum) -> type -- get the best matching Python type for the given Unreal enum" },
+	{ "register_python_shutdown_callback", PyCFunctionCast(&RegisterPythonShutdownCallback), METH_VARARGS, "register_python_shutdown_callback(callable: Callable[[], None]) -> object -- register the given callable (with no input arguments) as a callback to execute immediately before Python shutdown"},
+	{ "unregister_python_shutdown_callback", PyCFunctionCast(&UnregisterPythonShutdownCallback), METH_VARARGS, "unregister_python_shutdown_callback(handle: object) -> None -- unregister the given handle from a previous call to register_python_shutdown_callback"},
+	{ "NSLOCTEXT", PyCFunctionCast(&CreateLocalizedText), METH_VARARGS, "NSLOCTEXT(ns: str, key: str, source: str) -> Text -- create a localized Text from the given namespace, key, and source string" },
+	{ "LOCTABLE", PyCFunctionCast(&CreateLocalizedTextFromStringTable), METH_VARARGS, "LOCTABLE(id: Union[Name, str], key: str) -> Text -- get a localized Text from the given string table id and key" },
+	{ "is_editor", PyCFunctionCast(&IsEditor), METH_NOARGS, "is_editor() -> bool -- tells if the editor is running or not" },
+	{ "get_interpreter_executable_path", PyCFunctionCast(&GetInterpreterExecutablePath), METH_NOARGS, "get_interpreter_executable_path() -> str -- get the path to the Python interpreter executable of the Python SDK this plugin was compiled against" },
 	{ nullptr, nullptr, 0, nullptr }
 };
 
diff --git a/Engine/Plugins/Experimental/PythonScriptPlugin/Source/PythonScriptPlugin/Private/PyEditor.cpp b/Engine/Plugins/Experimental/PythonScriptPlugin/Source/PythonScriptPlugin/Private/PyEditor.cpp
--- a/Engine/Plugins/Experimental/PythonScriptPlugin/Source/PythonScriptPlugin/Private/PyEditor.cpp
+++ b/Engine/Plugins/Experimental/PythonScriptPlugin/Source/PythonScriptPlugin/Private/PyEditor.cpp
@@ -118,9 +118,9 @@ PyTypeObject InitializePyScopedEditorTransactionType()
 	};
 
 	static PyMethodDef PyMethods[] = {
-		{ "__enter__", PyCFunctionCast(&FMethods::EnterScope), METH_NOARGS, "x.__enter__() -> self -- begin this transaction" },
-		{ "__exit__", PyCFunctionCast(&FMethods::ExitScope), METH_VARARGS | METH_KEYWORDS, "x.__exit__(type, value, traceback) -> None -- end this transaction" },
-		{ "cancel", PyCFunctionCast(&FMethods::Cancel), METH_NOARGS, "x.cancel() -> None -- cancel this transaction" },
+		{ "__enter__", PyCFunctionCast(&FMethods::EnterScope), METH_NOARGS, "__enter__(self) -> ScopedEditorTransaction -- begin this transaction" },
+		{ "__exit__", PyCFunctionCast(&FMethods::ExitScope), METH_VARARGS | METH_KEYWORDS, "__exit__(self, type: Optional[Type[BaseException]], value: Optional[BaseException], traceback: Optional[TracebackType]) -> bool -- end this transaction" },
+		{ "cancel", PyCFunctionCast(&FMethods::Cancel), METH_NOARGS, "cancel(self) -> None -- cancel this transaction" },
 		{ nullptr, nullptr, 0, nullptr }
 	};
 
@@ -202,8 +202,8 @@ static PyObject* GetEditorSubsystem(FPyScopedEditorTransaction* InSelf, PyObject
 }
 
 PyMethodDef PyEditorMethods[] = {
-	{ "get_engine_subsystem", PyCFunctionCast(&PyEditor::GetEngineSubsystem), METH_VARARGS, "unreal.get_engine_subsystem() -> subsystem -- returns the requested subsystem could be null" },
-	{ "get_editor_subsystem", PyCFunctionCast(&PyEditor::GetEditorSubsystem), METH_VARARGS, "unreal.get_editor_subsystem() -> subsystem -- returns the requested subsystem could be null" },
+	{ "get_engine_subsystem", PyCFunctionCast(&PyEditor::GetEngineSubsystem), METH_VARARGS, "get_engine_subsystem(subsystem : Union[Class, type]) -> Any -- returns the requested Engine subsystem or None" },
+	{ "get_editor_subsystem", PyCFunctionCast(&PyEditor::GetEditorSubsystem), METH_VARARGS, "get_editor_subsystem(subsystem : Union[Class, type]) -> Any -- returns the requested Editor subsystem or None" },
 	{ nullptr, nullptr, 0, nullptr }
 };
 
diff --git a/Engine/Plugins/Experimental/PythonScriptPlugin/Source/PythonScriptPlugin/Private/PyEngine.cpp b/Engine/Plugins/Experimental/PythonScriptPlugin/Source/PythonScriptPlugin/Private/PyEngine.cpp
--- a/Engine/Plugins/Experimental/PythonScriptPlugin/Source/PythonScriptPlugin/Private/PyEngine.cpp
+++ b/Engine/Plugins/Experimental/PythonScriptPlugin/Source/PythonScriptPlugin/Private/PyEngine.cpp
@@ -201,7 +201,7 @@ PyObject* GetBlueprintGeneratedTypes(PyObject* InSelf, PyObject* InArgs)
 }
 
 PyMethodDef PyEngineMethods[] = {
-	{ "get_blueprint_generated_types", PyCFunctionCast(&GetBlueprintGeneratedTypes), METH_VARARGS, "x.get_blueprint_generated_types(...) -> tuple(type)/type/None -- get the Python types (will return a tuple for multiple types) for the given set of Blueprint asset paths (may be a sequence type or set of arguments)" },
+	{ "get_blueprint_generated_types", PyCFunctionCast(&GetBlueprintGeneratedTypes), METH_VARARGS, "get_blueprint_generated_types(paths: Iterable[str]) -> Optional[Union[type, Tuple[type, ...]]] -- get the Python types (will return a tuple for multiple types) for the given set of Blueprint asset paths (may be a sequence type or set of arguments)" },
 	{ nullptr, nullptr, 0, nullptr }
 };
 
diff --git a/Engine/Plugins/Experimental/PythonScriptPlugin/Source/PythonScriptPlugin/Private/PyGenUtil.cpp b/Engine/Plugins/Experimental/PythonScriptPlugin/Source/PythonScriptPlugin/Private/PyGenUtil.cpp
--- a/Engine/Plugins/Experimental/PythonScriptPlugin/Source/PythonScriptPlugin/Private/PyGenUtil.cpp
+++ b/Engine/Plugins/Experimental/PythonScriptPlugin/Source/PythonScriptPlugin/Private/PyGenUtil.cpp
@@ -77,8 +77,8 @@ const FName NativeMakeFuncMetaDataKey = TEXT("NativeMakeFunc");
 const FName ReturnValueKey = TEXT("ReturnValue");
 const TCHAR* HiddenMetaDataKey = TEXT("Hidden");
 
-
 TSharedPtr<IBreakIterator> NameBreakIterator;
+static ETypeHintingMode GTypeHintingMode = ETypeHintingMode::Off; // Use through SetTypeHintingMode()/IsTypeHintingEnabled()
 
 
 void FNativePythonModule::AddType(PyTypeObject* InType)
@@ -2165,14 +2165,14 @@ FParsedTooltip ParseTooltip(FStringView InTooltip)
 	return ParsedTooltip;
 }
 
-void PythonizeStructValueImpl(const UScriptStruct* InStruct, const void* InStructValue, const uint32 InFlags, FString& OutPythonDefaultValue);
+void PythonizeStructValueImpl(const UScriptStruct* InStruct, const void* InStructValue, const EPythonizeFlags InFlags, FString& OutPythonDefaultValue);
 
-void PythonizeValueImpl(const FProperty* InProp, const void* InPropValue, const uint32 InFlags, FString& OutPythonDefaultValue)
+void PythonizeValueImpl(const FProperty* InProp, const void* InPropValue, const EPythonizeFlags InFlags, FString& OutPythonDefaultValue)
 {
 	static const bool bIsForDocString = false;
 
-	const bool bIncludeUnrealNamespace = !!(InFlags & EPythonizeValueFlags::IncludeUnrealNamespace);
-	const bool bUseStrictTyping = !!(InFlags & EPythonizeValueFlags::UseStrictTyping);
+	const bool bIncludeUnrealNamespace = EnumHasAnyFlags(InFlags, EPythonizeFlags::IncludeUnrealNamespace);
+	const bool bUseStrictTyping = EnumHasAnyFlags(InFlags, EPythonizeFlags::UseStrictTyping);
 
 	const TCHAR* UnrealNamespace = bIncludeUnrealNamespace ? TEXT("unreal.") : TEXT("");
 
@@ -2350,12 +2350,12 @@ void PythonizeValueImpl(const FProperty* InProp, const void* InPropValue, const
 	}
 }
 
-void PythonizeStructValueImpl(const UScriptStruct* InStruct, const void* InStructValue, const uint32 InFlags, FString& OutPythonDefaultValue)
+void PythonizeStructValueImpl(const UScriptStruct* InStruct, const void* InStructValue, const EPythonizeFlags InFlags, FString& OutPythonDefaultValue)
 {
-	const bool bIncludeUnrealNamespace = !!(InFlags & EPythonizeValueFlags::IncludeUnrealNamespace);
-	const bool bUseStrictTyping = !!(InFlags & EPythonizeValueFlags::UseStrictTyping);
-	const bool bDefaultConstructStructs = !!(InFlags & EPythonizeValueFlags::DefaultConstructStructs);
-	const bool bDefaultConstructDateTime = !!(InFlags & EPythonizeValueFlags::DefaultConstructDateTime);
+	const bool bIncludeUnrealNamespace = EnumHasAnyFlags(InFlags, EPythonizeFlags::IncludeUnrealNamespace);
+	const bool bUseStrictTyping = EnumHasAnyFlags(InFlags, EPythonizeFlags::UseStrictTyping);
+	const bool bDefaultConstructStructs = EnumHasAnyFlags(InFlags, EPythonizeFlags::DefaultConstructStructs);
+	const bool bDefaultConstructDateTime = EnumHasAnyFlags(InFlags, EPythonizeFlags::DefaultConstructDateTime);
 
 	const TCHAR* UnrealNamespace = bIncludeUnrealNamespace ? TEXT("unreal.") : TEXT("");
 
@@ -2447,20 +2447,127 @@ void PythonizeStructValueImpl(const UScriptStruct* InStruct, const void* InStruc
 		: TEXT("]");
 }
 
-FString PythonizeValue(const FProperty* InProp, const void* InPropValue, const uint32 InFlags)
+FString PythonizeValue(const FProperty* InProp, const void* InPropValue, const EPythonizeFlags InFlags)
 {
 	FString PythonValue;
 	PythonizeValueImpl(InProp, InPropValue, InFlags, PythonValue);
 	return PythonValue;
 }
 
-FString PythonizeDefaultValue(const FProperty* InProp, const FString& InDefaultValue, const uint32 InFlags)
+FString PythonizeDefaultValue(const FProperty* InProp, const FString& InDefaultValue, const EPythonizeFlags InFlags)
 {
 	PyUtil::FPropValueOnScope PropValue(PyUtil::FConstPropOnScope::ExternalReference(InProp));
 	PyUtil::ImportDefaultValue(InProp, PropValue.GetValue(), InDefaultValue);
 	return PythonizeValue(InProp, PropValue.GetValue(), InFlags);
 }
 
+FString PythonizeMethodParam(const FGeneratedWrappedMethodParameter& MethodParam, EPythonizeFlags InPythonizeFlags, const TFunction<bool(const FString&)>& ShouldEllipseDefaultValue)
+{
+	bool bWithTypeHinting = EnumHasAnyFlags(InPythonizeFlags, EPythonizeFlags::WithTypeHinting);
+	bool bIncludeUnrealNamespace = EnumHasAnyFlags(InPythonizeFlags, EPythonizeFlags::IncludeUnrealNamespace);
+	bool bEllipseMissingDefaultValue = EnumHasAnyFlags(InPythonizeFlags, EPythonizeFlags::AddMissingDefaultValueEllipseWorkaround);
+
+	FString ParamName(UTF8_TO_TCHAR(MethodParam.ParamName.GetData()));
+	FString ParamType(bWithTypeHinting ? GetPropertyTypePythonName(MethodParam.ParamProp, bIncludeUnrealNamespace, /*bIsForDocString*/false, InPythonizeFlags) : FString());
+	FString DefaultValue;
+
+	if (bEllipseMissingDefaultValue && !MethodParam.ParamDefaultValue.IsSet())
+	{
+		DefaultValue = TEXT("..."); // Add the missing default value.
+	}
+	else if (MethodParam.ParamDefaultValue.IsSet())
+	{
+		DefaultValue = PyGenUtil::PythonizeDefaultValue(MethodParam.ParamProp, MethodParam.ParamDefaultValue.GetValue(), InPythonizeFlags);
+		if (ShouldEllipseDefaultValue && ShouldEllipseDefaultValue(DefaultValue))
+		{
+			DefaultValue = TEXT("..."); // In some context (like generated stub), the default value might refer to not-yet declared types, which is illegal
+		}
+	}
+
+	return PythonizeMethodParam(ParamName, ParamType, DefaultValue);
+}
+
+FString PythonizeMethodParam(const FString& InParamName, const FProperty* ParamProp, EPythonizeFlags InPythonizeFlags)
+{
+	bool bWithTypeHinting = EnumHasAnyFlags(InPythonizeFlags, EPythonizeFlags::WithTypeHinting);
+	bool bIncludeUnrealNamespace = EnumHasAnyFlags(InPythonizeFlags, EPythonizeFlags::IncludeUnrealNamespace);
+
+	FString ParamType(bWithTypeHinting ? GetPropertyTypePythonName(ParamProp, bIncludeUnrealNamespace, /*bIsForDocString*/false, InPythonizeFlags) : FString());
+
+	return PythonizeMethodParam(InParamName, ParamType);
+}
+
+FString PythonizeMethodParam(const FString& InParamName, const FString& InParamType, const FString& InDefaultValue)
+{
+	return InDefaultValue.IsEmpty()
+		? InParamType.IsEmpty() ? FString::Printf(TEXT("%s"), *InParamName) : FString::Printf(TEXT("%s: %s"), *InParamName, *InParamType)
+		: InParamType.IsEmpty() ? FString::Printf(TEXT("%s = %s"), *InParamName, *InDefaultValue) : FString::Printf(TEXT("%s: %s = %s"), *InParamName, *InParamType, *InDefaultValue);
+}
+
+FString PythonizeMethodReturnType(const TArray<PyGenUtil::FGeneratedWrappedMethodParameter>& InOutputParams, EPythonizeFlags InPythonizeFlags)
+{
+	// If type hinting is disabled, return type is not specified with the method signature.
+	if (!EnumHasAnyFlags(InPythonizeFlags, EPythonizeFlags::WithTypeHinting))
+	{
+		return FString();
+	}
+	if (InOutputParams.Num() == 0)
+	{
+		return TEXT("None");
+	}
+
+	bool bIncludeUnrealNamespace = EnumHasAnyFlags(InPythonizeFlags, EPythonizeFlags::IncludeUnrealNamespace);
+	bool bCanReturnNone = false;
+
+	// If we have multiple return values and the main return value is a bool, skip it (to mimic PyGenUtils::PackReturnValues)
+	int32 ReturnPropIndex = 0;
+	if (InOutputParams.Num() > 1 && InOutputParams[0].ParamProp->HasAnyPropertyFlags(CPF_ReturnParm) && InOutputParams[0].ParamProp->IsA<FBoolProperty>())
+	{
+		ReturnPropIndex = 1; // Start packing at the 1st out value
+		bCanReturnNone = true; // When the C++ function returns false, the Python methods returns None, otherwise, it returns the output parameter(s) as a single value or a tuple.
+	}
+
+	// Do we need to return a packed tuple, or just a single value?
+	const int32 NumPropertiesToPack = InOutputParams.Num() - ReturnPropIndex;
+	if (NumPropertiesToPack == 1)
+	{
+		const PyGenUtil::FGeneratedWrappedMethodParameter& ReturnParam = InOutputParams[ReturnPropIndex];
+		if (bCanReturnNone)
+		{
+			return FString::Printf(TEXT("Optional[%s]"), *PyGenUtil::GetPropertyTypePythonName(ReturnParam.ParamProp, bIncludeUnrealNamespace, /*bIsForDocString*/false, InPythonizeFlags));
+		}
+		return PyGenUtil::GetPropertyTypePythonName(ReturnParam.ParamProp, bIncludeUnrealNamespace, /*bIsForDocString*/false, InPythonizeFlags);
+	}
+	else // Returns a tuple
+	{
+		FString ReturnType;
+		ReturnType = bCanReturnNone ? TEXT("Optional[Tuple[") : TEXT("Tuple[");
+		for (int32 PackedPropIndex = 0; ReturnPropIndex < InOutputParams.Num(); ++ReturnPropIndex, ++PackedPropIndex)
+		{
+			if (PackedPropIndex > 0)
+			{
+				ReturnType += TEXT(", ");
+			}
+			const PyGenUtil::FGeneratedWrappedMethodParameter& ReturnParam = InOutputParams[ReturnPropIndex];
+			ReturnType += PyGenUtil::GetPropertyTypePythonName(ReturnParam.ParamProp, bIncludeUnrealNamespace, /*bIsForDocString*/false, InPythonizeFlags);
+		}
+		ReturnType += (bCanReturnNone ? TEXT("]]") : TEXT("]"));
+		return ReturnType;
+	}
+}
+
+FString PythonizeMethodReturnType(const FProperty* InProp, EPythonizeFlags InPythonizeFlags)
+{
+	// If type hinting is disabled, return type is not specified with the method signature.
+	if (!EnumHasAnyFlags(InPythonizeFlags, EPythonizeFlags::WithTypeHinting))
+	{
+		return FString();
+	}
+
+	bool bIncludeUnrealNamespace = EnumHasAnyFlags(InPythonizeFlags, EPythonizeFlags::IncludeUnrealNamespace);
+	return PyGenUtil::GetPropertyTypePythonName(InProp, bIncludeUnrealNamespace, /*bIsForDocString*/false, InPythonizeFlags);
+}
+
 const UObject* GetAssetTypeRegistryType(const UObject* InObj)
 {
 	if (const UBlueprintCore* BlueprintAsset = Cast<const UBlueprintCore>(InObj))
@@ -2993,16 +3100,26 @@ TArray<FString> GetDeprecatedPropertyPythonNames(const FProperty* InProp)
 	return PropNames;
 }
 
-FString GetPropertyTypePythonName(const FProperty* InProp, const bool InIncludeUnrealNamespace, const bool InIsForDocString)
+FString GetPropertyTypePythonName(const FProperty* InProp, const bool InIncludeUnrealNamespace, const bool InIsForDocString, EPythonizeFlags InPythonizeFlags)
 {
+	const TCHAR* UnrealNamespace = InIncludeUnrealNamespace ? TEXT("unreal.") : TEXT("");
+
+	auto GetTypeWithCoercionHint = [InPythonizeFlags, UnrealNamespace](const TCHAR* TypeName) -> FString
+	{
+		if (EnumHasAnyFlags(InPythonizeFlags, EPythonizeFlags::WithTypeCoercion))
+		{
+			// Supports using a python string in place of a unreal.Name() or unreal.Text().
+			return FString::Printf(TEXT("Union[%s%s, str]"), UnrealNamespace, TypeName);
+		}
+		return FString::Printf(TEXT("%s%s"), UnrealNamespace, TypeName);
+	};
+
 #define GET_PROPERTY_TYPE(TYPE, VALUE)				\
 		if (CastField<const TYPE>(InProp) != nullptr)	\
 		{											\
 			return (VALUE);							\
 		}
 
-	const TCHAR* UnrealNamespace = InIncludeUnrealNamespace ? TEXT("unreal.") : TEXT("");
-
 	GET_PROPERTY_TYPE(FBoolProperty, TEXT("bool"))
 	GET_PROPERTY_TYPE(FInt8Property, InIsForDocString ? TEXT("int8") : TEXT("int"))
 	GET_PROPERTY_TYPE(FInt16Property, InIsForDocString ? TEXT("int16") : TEXT("int"))
@@ -3014,8 +3131,8 @@ FString GetPropertyTypePythonName(const FProperty* InProp, const bool InIncludeU
 	GET_PROPERTY_TYPE(FFloatProperty, TEXT("float"))
 	GET_PROPERTY_TYPE(FDoubleProperty, InIsForDocString ? TEXT("double") : TEXT("float"))
 	GET_PROPERTY_TYPE(FStrProperty, TEXT("str"))
-	GET_PROPERTY_TYPE(FNameProperty, InIncludeUnrealNamespace ? TEXT("unreal.Name") : TEXT("Name"))
-	GET_PROPERTY_TYPE(FTextProperty, InIncludeUnrealNamespace ? TEXT("unreal.Text") : TEXT("Text"))
+	GET_PROPERTY_TYPE(FNameProperty, GetTypeWithCoercionHint(TEXT("Name")))
+	GET_PROPERTY_TYPE(FTextProperty, GetTypeWithCoercionHint(TEXT("Text")))
 	GET_PROPERTY_TYPE(FFieldPathProperty, InIncludeUnrealNamespace ? TEXT("unreal.FieldPath") : TEXT("FieldPath"))
 
 	if (const FByteProperty* ByteProp = CastField<const FByteProperty>(InProp))
@@ -3033,24 +3150,36 @@ FString GetPropertyTypePythonName(const FProperty* InProp, const bool InIncludeU
 	{
 		return FString::Printf(TEXT("%s%s"), UnrealNamespace, *GetEnumPythonName(EnumProp->GetEnum()));
 	}
-	if (InIsForDocString)
+	if (const FClassProperty* ClassProp = CastField<const FClassProperty>(InProp))
 	{
-		if (const FClassProperty* ClassProp = CastField<const FClassProperty>(InProp))
+		if (InIsForDocString)
 		{
 			return FString::Printf(TEXT("type(%s%s)"), UnrealNamespace, *GetClassPythonName(ClassProp->PropertyClass));
 		}
+		else if (EnumHasAnyFlags(InPythonizeFlags, EPythonizeFlags::WithTypeCoercion))
+		{
+			// If type hinting is enabled and the context allows type coercion (for an input parameters), the implementation support taking a 'type' in place of a 'Class'.
+			return FString::Printf(TEXT("Union[%sClass, type]"), UnrealNamespace);
+		}
 	}
 	if (const FObjectPropertyBase* ObjProp = CastField<const FObjectPropertyBase>(InProp))
 	{
-		return FString::Printf(TEXT("%s%s"), UnrealNamespace, *GetClassPythonName(ObjProp->PropertyClass));
+		return (EnumHasAnyFlags(InPythonizeFlags, EPythonizeFlags::WithOptionalType))
+			? FString::Printf(TEXT("Optional[%s%s]"), UnrealNamespace, *GetClassPythonName(ObjProp->PropertyClass)) // None can be assigned to object.
+			: FString::Printf(TEXT("%s%s"), UnrealNamespace, *GetClassPythonName(ObjProp->PropertyClass));
 	}
 	if (const FInterfaceProperty* InterfaceProp = CastField<const FInterfaceProperty>(InProp))
 	{
-		return FString::Printf(TEXT("%s%s"), UnrealNamespace, *GetClassPythonName(InterfaceProp->InterfaceClass));
+		return (EnumHasAnyFlags(InPythonizeFlags, EPythonizeFlags::WithOptionalType))
+			? FString::Printf(TEXT("Optional[%s%s]"), UnrealNamespace, *GetClassPythonName(InterfaceProp->InterfaceClass)) // None can be assigned to Interface.
+			: FString::Printf(TEXT("%s%s"), UnrealNamespace, *GetClassPythonName(InterfaceProp->InterfaceClass));
 	}
 	if (const FStructProperty* StructProp = CastField<const FStructProperty>(InProp))
 	{
-		return FString::Printf(TEXT("%s%s"), UnrealNamespace, *GetStructPythonName(StructProp->Struct));
+		// If type hinting is enabled and the context allows type coercion (for an input parameters), the implementation support several types in place of the struct.
+		return (EnumHasAnyFlags(InPythonizeFlags, EPythonizeFlags::WithTypeCoercion))
+			? FString::Printf(TEXT("Union[%s%s, Iterable[Any], Mapping[str, Any]]"), UnrealNamespace, *GetStructPythonName(StructProp->Struct))
+			: FString::Printf(TEXT("%s%s"), UnrealNamespace, *GetStructPythonName(StructProp->Struct));
 	}
 	if (const FDelegateProperty* DelegateProp = CastField<const FDelegateProperty>(InProp))
 	{
@@ -3062,15 +3191,35 @@ FString GetPropertyTypePythonName(const FProperty* InProp, const bool InIncludeU
 	}
 	if (const FArrayProperty* ArrayProperty = CastField<const FArrayProperty>(InProp))
 	{
-		return FString::Printf(TEXT("%sArray(%s)"), UnrealNamespace, *GetPropertyTypePythonName(ArrayProperty->Inner, InIncludeUnrealNamespace, InIsForDocString));
+		// If type hinting is enabled and the context allows type coercion (for an input parameters), the implementation support any Iterable[T] in place of an Array[T]
+		return EnumHasAnyFlags(InPythonizeFlags, EPythonizeFlags::WithTypeCoercion)
+			? FString::Printf(TEXT("Iterable[%s]"), *GetPropertyTypePythonName(ArrayProperty->Inner, InIncludeUnrealNamespace, InIsForDocString, InPythonizeFlags))
+			: FString::Printf(TEXT("Array[%s]"), *GetPropertyTypePythonName(ArrayProperty->Inner, InIncludeUnrealNamespace, InIsForDocString, InPythonizeFlags));
+
 	}
 	if (const FSetProperty* SetProperty = CastField<const FSetProperty>(InProp))
 	{
-		return FString::Printf(TEXT("%sSet(%s)"), UnrealNamespace, *GetPropertyTypePythonName(SetProperty->ElementProp, InIncludeUnrealNamespace, InIsForDocString));
+		// If type hinting is enabled and the context allows type coercion (for an input parameters), the implementation support any Iterable[T] in place of a Set[T]
+		return EnumHasAnyFlags(InPythonizeFlags, EPythonizeFlags::WithTypeCoercion)
+			? FString::Printf(TEXT("Iterable[%s]"), *GetPropertyTypePythonName(SetProperty->ElementProp, InIncludeUnrealNamespace, InIsForDocString, InPythonizeFlags))
+			: FString::Printf(TEXT("%sSet[%s]"), UnrealNamespace, *GetPropertyTypePythonName(SetProperty->ElementProp, InIncludeUnrealNamespace, InIsForDocString, InPythonizeFlags));
 	}
 	if (const FMapProperty* MapProperty = CastField<const FMapProperty>(InProp))
 	{
-		return FString::Printf(TEXT("%sMap(%s, %s)"), UnrealNamespace, *GetPropertyTypePythonName(MapProperty->KeyProp, InIncludeUnrealNamespace, InIsForDocString), *GetPropertyTypePythonName(MapProperty->ValueProp, InIncludeUnrealNamespace, InIsForDocString));
+		// If type hinting is enabled and the context allows type coercion (for an input parameters).
+		if (EnumHasAnyFlags(InPythonizeFlags, EPythonizeFlags::WithTypeCoercion))
+		{
+			// When an unreal.Map() is expected, the implementation will successfully parse Mapping, list of 2-element tuples or even list of list.
+			FString KeyTypeDecl   = GetPropertyTypePythonName(MapProperty->KeyProp, InIncludeUnrealNamespace, InIsForDocString, InPythonizeFlags);
+			FString ValueTypeDecl = GetPropertyTypePythonName(MapProperty->ValueProp, InIncludeUnrealNamespace, InIsForDocString, InPythonizeFlags);
+			return FString::Printf(TEXT("Union[Mapping[%s, %s], Iterable[Tuple[%s, %s]], List[List[Union[%s, %s]]]]"), *KeyTypeDecl, *ValueTypeDecl, *KeyTypeDecl, *ValueTypeDecl, *KeyTypeDecl, *ValueTypeDecl);
+		}
+		else
+		{
+			return FString::Printf(TEXT("%sMap[%s, %s]"), UnrealNamespace,
+				*GetPropertyTypePythonName(MapProperty->KeyProp, InIncludeUnrealNamespace, InIsForDocString, InPythonizeFlags),
+				*GetPropertyTypePythonName(MapProperty->ValueProp, InIncludeUnrealNamespace, InIsForDocString, InPythonizeFlags));
+		}
 	}
 
 	return InIsForDocString ? TEXT("'undefined'") : TEXT("type");
@@ -3201,6 +3350,27 @@ bool SaveGeneratedTextFile(const TCHAR* InFilename, FStringView InFileContents,
 	return true;
 }
 
+void SetTypeHintingMode(ETypeHintingMode InMode)
+{
+	GTypeHintingMode = InMode;
+}
+
+ETypeHintingMode GetTypeHintingMode()
+{
+	if (!IsTypeHintingEnabled()) // Can be disable if the Python version is too old.
+	{
+		return ETypeHintingMode::Off;
+	}
+	return GTypeHintingMode;
+}
+
+bool IsTypeHintingEnabled()
+{
+	// In 3.7, a feature allowing usage of non-declared types in type hints was added (from __future__ import annotations) and the engine
+	// doesn't output types in a strict declaration/usage order, so this is required to correctly type.
+	return PY_MAJOR_VERSION >= 3 && PY_MINOR_VERSION >= 7 && GTypeHintingMode != ETypeHintingMode::Off;
+}
+
 }
 
 #endif	// WITH_PYTHON
diff --git a/Engine/Plugins/Experimental/PythonScriptPlugin/Source/PythonScriptPlugin/Private/PyGenUtil.h b/Engine/Plugins/Experimental/PythonScriptPlugin/Source/PythonScriptPlugin/Private/PyGenUtil.h
--- a/Engine/Plugins/Experimental/PythonScriptPlugin/Source/PythonScriptPlugin/Private/PyGenUtil.h
+++ b/Engine/Plugins/Experimental/PythonScriptPlugin/Source/PythonScriptPlugin/Private/PyGenUtil.h
@@ -6,6 +6,7 @@
 #include "PyPtr.h"
 #include "PyConstant.h"
 #include "PyMethodWithClosure.h"
+#include "PythonScriptPluginSettings.h"
 #include "CoreMinimal.h"
 #include "UObject/ObjectMacros.h"
 #include "UObject/WeakObjectPtr.h"
@@ -866,18 +867,28 @@ namespace PyGenUtil
 		Upper,
 	};
 
-	/** Flags controlling the behavior of PythonizeValue */
-	namespace EPythonizeValueFlags
+	/** Flags controlling the behavior of PythonizeValue/PythonizeMethodParam/PythonizeMethodReturnType. */
+	enum class EPythonizeFlags
 	{
-		enum EFlags
-		{
-			None = 0,
-			IncludeUnrealNamespace = 1<<0,
-			UseStrictTyping = 1<<1,
-			DefaultConstructStructs = 1<<2,
-			DefaultConstructDateTime = 1<<3,
-		};
+		None = 0,
+		IncludeUnrealNamespace = 1<<0,
+		UseStrictTyping = 1<<1,
+		DefaultConstructStructs = 1<<2,
+		DefaultConstructDateTime = 1<<3,
+
+		/** Type hinting is enabled. */
+		WithTypeHinting = 1<<4,
+
+		/** If the type of a value being hinted can be None, wrap the type in Optional[T], which is an alias to Union[T, None]. */
+		WithOptionalType = 1<<5,
+
+		/** Declare the type along with it known type coercion. For example [unreal.Name, str] because a string can be passed in place of a Name. */
+		WithTypeCoercion = 1<<6,
+
+		/** In method declaration, add '...' as default value even if no default value was found. Used in stub generation to workaround 'non-defaulted param after defaulted param' coming from the 'malformed' UFUNCTION. */
+		AddMissingDefaultValueEllipseWorkaround = 1<<7,
 	};
+	ENUM_CLASS_FLAGS(EPythonizeFlags)
 
 	/** Context information passed to PythonizeTooltip */
 	struct FPythonizeTooltipContext
@@ -1086,10 +1097,25 @@ namespace PyGenUtil
 	FParsedTooltip ParseTooltip(FStringView InTooltip);
 
 	/** Given a property and its value, convert it into something that could be used in a Python script */
-	FString PythonizeValue(const FProperty* InProp, const void* InPropValue, const uint32 InFlags = EPythonizeValueFlags::None);
+	FString PythonizeValue(const FProperty* InProp, const void* InPropValue, const EPythonizeFlags InPythonizeFlags = EPythonizeFlags::None);
 
 	/** Given a property and its default value, convert it into something that could be used in a Python script */
-	FString PythonizeDefaultValue(const FProperty* InProp, const FString& InDefaultValue, const uint32 InFlags = EPythonizeValueFlags::None);
+	FString PythonizeDefaultValue(const FProperty* InProp, const FString& InDefaultValue, const EPythonizeFlags InPythonizeFlags = EPythonizeFlags::None);
+
+	/** Given a generated method parameter, extract the name, type and default value and format a parameter declaration usable in a Python method signature according to the specified flags. */
+	FString PythonizeMethodParam(const FGeneratedWrappedMethodParameter& InMethodParam, EPythonizeFlags InPythonizeFlags = EPythonizeFlags::None, const TFunction<bool(const FString&)>& ShouldEllipseDefaultValue = TFunction<bool(const FString&)>());
+
+	/** Given a property and a name, format a parameter declaration usable in a Python method signature according to the specified flags. */
+	FString PythonizeMethodParam(const FString& InParamName, const FProperty* ParamProp, EPythonizeFlags InPythonizeFlags = EPythonizeFlags::None);
+
+	/** Given a param name and optional param type and default value, format a parameter declaration usable in a Python method signature. */
+	FString PythonizeMethodParam(const FString& InParamName, const FString& InParamType = FString(), const FString& InDefaultValue = FString());
+
+	/** Given list of return parameter, format the return type declaration usable in the Python method signature according to the flags. */
+	FString PythonizeMethodReturnType(const TArray<PyGenUtil::FGeneratedWrappedMethodParameter>& InOutputParams, EPythonizeFlags InPythonizeFlags = EPythonizeFlags::None);
+
+	/** Given a property, format a return type declaration usable in the Python method signature according to the flags. */
+	FString PythonizeMethodReturnType(const FProperty* InProp, EPythonizeFlags InPythonizeFlags = EPythonizeFlags::None);
 
 	/** Get the type that should be used with the Python type registry for an asset (eg, a Blueprint asset should use its generated class) */
 	const UObject* GetAssetTypeRegistryType(const UObject* InObj);
@@ -1167,7 +1193,7 @@ namespace PyGenUtil
 	TArray<FString> GetDeprecatedPropertyPythonNames(const FProperty* InProp);
 
 	/** Get the Python name of the given property */
-	FString GetPropertyTypePythonName(const FProperty* InProp, const bool InIncludeUnrealNamespace = false, const bool InIsForDocString = true);
+	FString GetPropertyTypePythonName(const FProperty* InProp, const bool InIncludeUnrealNamespace = false, const bool InIsForDocString = true, const EPythonizeFlags PythonizeFlags = EPythonizeFlags::None);
 
 	/** Get the Python type of the given property */
 	FString GetPropertyPythonType(const FProperty* InProp);
@@ -1193,6 +1219,26 @@ namespace PyGenUtil
 
 	/** Save a generated text file to disk as UTF-8 (only writes the file if the contents differs, unless forced) */
 	bool SaveGeneratedTextFile(const TCHAR* InFilename, FStringView InFileContents, const bool InForceWrite = false);
+
+	/**
+	 * Whether Python type hints should be generated. This is only effective if Python interpreter version has the support. This
+	 * affects the Python stub and docstrings. The function must be called before Python glue and the stub are generated. It
+	 * can be set from the user settings or forced to a value for a specific purpose, like generating the online doc.
+	 */
+	void SetTypeHintingMode(ETypeHintingMode InMode);
+
+	/**
+	 * Return the effecting type hinting mode.
+	 * @see SetTypeHintingMode.
+	 */
+	ETypeHintingMode GetTypeHintingMode();
+
+	/**
+	 * Whether the stub and docstrings are type hinted during glue/stub generation. Default to false. This doesn't read
+	 * the user settings, but rather the value set with SetTypeHintingMode(). The function always returns false if
+	 * the linked Python interpreter version is older than 3.7.
+	 */
+	bool IsTypeHintingEnabled();
 }
 
 #endif	// WITH_PYTHON
diff --git a/Engine/Plugins/Experimental/PythonScriptPlugin/Source/PythonScriptPlugin/Private/PySlate.cpp b/Engine/Plugins/Experimental/PythonScriptPlugin/Source/PythonScriptPlugin/Private/PySlate.cpp
--- a/Engine/Plugins/Experimental/PythonScriptPlugin/Source/PythonScriptPlugin/Private/PySlate.cpp
+++ b/Engine/Plugins/Experimental/PythonScriptPlugin/Source/PythonScriptPlugin/Private/PySlate.cpp
@@ -197,11 +197,11 @@ PyObject* ParentExternalWindowToSlate(PyObject* InSelf, PyObject* InArgs)
 }
 
 PyMethodDef PySlateMethods[] = {
-	{ "register_slate_pre_tick_callback", PyCFunctionCast(&RegisterSlatePreTickCallback), METH_VARARGS, "x.register_slate_pre_tick_callback(callable) -> _DelegateHandle -- register the given callable (taking a single float) as a pre-tick callback in Slate" },
-	{ "unregister_slate_pre_tick_callback", PyCFunctionCast(&UnregisterSlatePreTickCallback), METH_VARARGS, "x.unregister_slate_pre_tick_callback(handle) -> None -- unregister the given handle from a previous call to register_slate_pre_tick_callback" },
-	{ "register_slate_post_tick_callback", PyCFunctionCast(&RegisterSlatePostTickCallback), METH_VARARGS, "x.register_slate_post_tick_callback(callable) -> _DelegateHandle -- register the given callable (taking a single float) as a pre-tick callback in Slate" },
-	{ "unregister_slate_post_tick_callback", PyCFunctionCast(&UnregisterSlatePostTickCallback), METH_VARARGS, "x.unregister_slate_post_tick_callback(handle) -> None -- unregister the given handle from a previous call to register_slate_post_tick_callback" },
-	{ "parent_external_window_to_slate", PyCFunctionCast(&ParentExternalWindowToSlate), METH_VARARGS, "x.parent_external_window_to_slate(external_window, parent_search_method=SlateParentWindowSearchMethod.ACTIVE_WINDOW) -> None -- parent the given OS specific external window handle to a suitable Slate window" },
+	{ "register_slate_pre_tick_callback", PyCFunctionCast(&RegisterSlatePreTickCallback), METH_VARARGS, "x.register_slate_pre_tick_callback(callable: Union[Callable[[float], None], DelegateBase]) -> object -- register the given callable (taking a single float) as a pre-tick callback in Slate" },
+	{ "unregister_slate_pre_tick_callback", PyCFunctionCast(&UnregisterSlatePreTickCallback), METH_VARARGS, "x.unregister_slate_pre_tick_callback(handle: object) -> None -- unregister the given handle from a previous call to register_slate_pre_tick_callback" },
+	{ "register_slate_post_tick_callback", PyCFunctionCast(&RegisterSlatePostTickCallback), METH_VARARGS, "x.register_slate_post_tick_callback(callable: Union[Callable[[float], None], DelegateBase]) -> object -- register the given callable (taking a single float) as a pre-tick callback in Slate" },
+	{ "unregister_slate_post_tick_callback", PyCFunctionCast(&UnregisterSlatePostTickCallback), METH_VARARGS, "x.unregister_slate_post_tick_callback(handle: object) -> None -- unregister the given handle from a previous call to register_slate_post_tick_callback" },
+	{ "parent_external_window_to_slate", PyCFunctionCast(&ParentExternalWindowToSlate), METH_VARARGS, "x.parent_external_window_to_slate(external_window_handle: object, parent_search_method: SlateParentWindowSearchMethod = SlateParentWindowSearchMethod.ACTIVE_WINDOW) -> None -- parent the given OS specific external window handle to a suitable Slate window" },
 	{ nullptr, nullptr, 0, nullptr }
 };
 
diff --git a/Engine/Plugins/Experimental/PythonScriptPlugin/Source/PythonScriptPlugin/Private/PyTest.cpp b/Engine/Plugins/Experimental/PythonScriptPlugin/Source/PythonScriptPlugin/Private/PyTest.cpp
--- a/Engine/Plugins/Experimental/PythonScriptPlugin/Source/PythonScriptPlugin/Private/PyTest.cpp
+++ b/Engine/Plugins/Experimental/PythonScriptPlugin/Source/PythonScriptPlugin/Private/PyTest.cpp
@@ -205,3 +205,111 @@ int32 UPyTestObjectLibrary::GetOtherConstantValue()
 {
 	return 20;
 }
+
+
+FString UPyTestTypeHint::GetStringConst()
+{
+	return FString("Foo");
+}
+
+int32 UPyTestTypeHint::GetIntConst()
+{
+	return 777;
+}
+
+
+UPyTestTypeHint::UPyTestTypeHint()
+{
+	ObjectProp = NewObject<UPyTestObject>();
+}
+
+UPyTestTypeHint::UPyTestTypeHint(bool bParam1, int32 Param2, float Param3, const FString& Param4, const FText& Param5)
+	: BoolProp(bParam1)
+	, IntProp(Param2)
+	, FloatProp(Param3)
+	, StringProp(Param4)
+	, TextProp(Param5)
+{
+}
+
+bool UPyTestTypeHint::CheckBoolTypeHints(bool bParam1, bool bParam2, bool bParam3)
+{
+	return true;
+}
+
+int32 UPyTestTypeHint::CheckIntegerTypeHints(uint8 Param1, int32 Param2, int64 Param3)
+{
+	return 0;
+}
+
+double UPyTestTypeHint::CheckFloatTypeHints(float Param1, double Param2, float Param3, double Param4)
+{
+	return 0.0;
+}
+
+EPyTestEnum UPyTestTypeHint::CheckEnumTypeHints(EPyTestEnum Param1, EPyTestEnum Param2)
+{
+	return EPyTestEnum::One;
+}
+
+FString UPyTestTypeHint::CheckStringTypeHints(const FString& Param1, const FString& Param2)
+{
+	return FString();
+}
+
+FName UPyTestTypeHint::CheckNameTypeHints(const FName& Param1, const FName& Param2)
+{
+	return FName();
+}
+
+FText UPyTestTypeHint::CheckTextTypeHints(const FText& Param1, const FText& Param2)
+{
+	return FText::GetEmpty();
+}
+
+TFieldPath<FProperty> UPyTestTypeHint::CheckFieldPathTypeHints(const TFieldPath<FProperty> Param1)
+{
+	return TFieldPath<FProperty>();
+}
+
+FPyTestStruct UPyTestTypeHint::CheckStructTypeHints(const FPyTestStruct& Param1, const FPyTestStruct& Param2)
+{
+	return FPyTestStruct();
+}
+
+UPyTestObject* UPyTestTypeHint::CheckObjectTypeHints(const UPyTestObject* Param1, const UPyTestObject* Param3)
+{
+	return nullptr;
+}
+
+TArray<FText> UPyTestTypeHint::CheckArrayTypeHints(const TArray<FString>& Param1, const TArray<FName>& Param2, const TArray<FText>& Param3, const TArray<UObject*>& Param4)
+{
+	return TArray<FText>();
+}
+
+TSet<FName> UPyTestTypeHint::CheckSetTypeHints(const TSet<FString>& Param1, const TSet<FName>& Param2, const TSet<UObject*>& Param4)
+{
+	return TSet<FName>();
+}
+
+TMap<FString, UObject*> UPyTestTypeHint::CheckMapTypeHints(const TMap<int, FString>& Param1, const TMap<int, FName>& Param2, const TMap<int, FText>& Param3, const TMap<int, UObject*>& Param4)
+{
+	return TMap<FString, UObject*>();
+}
+
+FPyTestDelegate& UPyTestTypeHint::CheckDelegateTypeHints(const FPyTestDelegate& Param1)
+{
+	return DelegateProp;
+}
+
+bool UPyTestTypeHint::CheckStaticFunction(bool Param1, int32 Param2, double Param3, const FString& Param4)
+{
+	return true;
+}
+
+int UPyTestTypeHint::CheckTupleReturnType(UPARAM(ref) FString& InOutString)
+{
+	InOutString = TEXT("Foo");
+	return 0;
+}
+
diff --git a/Engine/Plugins/Experimental/PythonScriptPlugin/Source/PythonScriptPlugin/Private/PyTest.h b/Engine/Plugins/Experimental/PythonScriptPlugin/Source/PythonScriptPlugin/Private/PyTest.h
--- a/Engine/Plugins/Experimental/PythonScriptPlugin/Source/PythonScriptPlugin/Private/PyTest.h
+++ b/Engine/Plugins/Experimental/PythonScriptPlugin/Source/PythonScriptPlugin/Private/PyTest.h
@@ -17,6 +17,11 @@ DECLARE_DYNAMIC_DELEGATE_RetVal_OneParam(int32, FPyTestDelegate, int32, InValue)
  */
 DECLARE_DYNAMIC_MULTICAST_DELEGATE_OneParam(FPyTestMulticastDelegate, FString, InStr);
 
+/**
+ * Delegate for slate pre/post tick event.
+ */
+DECLARE_DYNAMIC_DELEGATE_OneParam(FPyTestSlateTickDelegate, float, InDeltaTime);
+
 /**
  * Enum to allow testing of the various UEnum features that are exposed to Python wrapped types.
  */
@@ -320,3 +325,151 @@ public:
 	UPROPERTY(BlueprintAssignable, Category = "Python|Internal")
 	FOnNameCollisionDelegate OnNameCollisionTestDelegate; // Same prop name and type name as UPyTestStructDelegate::FOnNameCollisionTestDelegate, but different params.
 };
+
+
+/* Used to verify if the generated Python stub is correctly type-hinted (if type hint is enabled). The stub is generated
+ * in the project intermediate folder when the Python developer mode is enabled (Editor preferences). The type hints can
+ * be checked in the stub itself or PythonScriptPlugin/Content/Python/test_type_hints.py can be loaded in a Python IDE that
+ * supports type checking and look at the code to verify that there is not problems with the types.
+ */
+UCLASS(Blueprintable)
+class UPyTestTypeHint : public UObject
+{
+	GENERATED_BODY()
+
+public:
+	//
+	// Check type hinted init methods.
+	//
+
+	UPyTestTypeHint();
+	UPyTestTypeHint(bool bParam1, int32 Param2, float Param3, const FString& Param4 = "Hi", const FText& Param5 = FText());
+
+	//
+	// Check type hinted constants
+	//
+
+	UFUNCTION(BlueprintPure, Category = "Python|Internal", meta=(ScriptConstant="StrConst"))
+	static FString GetStringConst();
+
+	UFUNCTION(BlueprintPure, Category = "Python|Internal", meta=(ScriptConstant="IntConst"))
+	static int32 GetIntConst();
+
+	//
+	// Check type hinted properties (setter/getter)
+	//
+
+	UPROPERTY(EditAnywhere, BlueprintReadWrite, Category = "Python|Internal")
+	bool BoolProp = true;
+
+	UPROPERTY(EditAnywhere, BlueprintReadWrite, Category = "Python|Internal")
+	int32 IntProp = 32;
+
+	UPROPERTY(EditAnywhere, BlueprintReadWrite, Category = "Python|Internal")
+	float FloatProp;
+
+	UPROPERTY(EditAnywhere, BlueprintReadWrite, Category = "Python|Internal")
+	EPyTestEnum EnumProp;
+
+	UPROPERTY(EditAnywhere, BlueprintReadWrite, Category = "Python|Internal")
+	FString StringProp;
+
+	UPROPERTY(EditAnywhere, BlueprintReadWrite, Category = "Python|Internal")
+	FName NameProp;
+
+	UPROPERTY(EditAnywhere, BlueprintReadWrite, Category = "Python|Internal")
+	FText TextProp;
+
+	UPROPERTY(EditAnywhere, BlueprintReadWrite, Category = "Python|Internal")
+	TFieldPath<FProperty> FieldPathProp;
+
+	UPROPERTY(EditAnywhere, BlueprintReadWrite, Category = "Python|Internal")
+	FPyTestStruct StructProp;
+
+	UPROPERTY(EditAnywhere, BlueprintReadWrite, Category = "Python|Internal")
+	UPyTestObject* ObjectProp = nullptr;
+
+	UPROPERTY(EditAnywhere, BlueprintReadWrite, Category = "Python|Internal")
+	TArray<FString> StrArrayProp;
+
+	UPROPERTY(EditAnywhere, BlueprintReadWrite, Category = "Python|Internal")
+	TArray<FName> NameArrayProp;
+
+	UPROPERTY(EditAnywhere, BlueprintReadWrite, Category = "Python|Internal")
+	TArray<FText> TextArrayProp;
+
+	UPROPERTY(EditAnywhere, BlueprintReadWrite, Category = "Python|Internal")
+	TArray<UObject*> ObjectArrayProp;
+
+	UPROPERTY(EditAnywhere, BlueprintReadWrite, Category = "Python|Internal")
+	TSet<FString> SetProp;
+
+	UPROPERTY(EditAnywhere, BlueprintReadWrite, Category = "Python|Internal")
+	TMap<int32, FString> MapProp;
+
+	UPROPERTY(EditAnywhere, BlueprintReadWrite, Category = "Python|Internal")
+	FPyTestDelegate DelegateProp;
+
+	UPROPERTY(EditAnywhere, BlueprintAssignable, Category = "Python|Internal")
+	FPyTestMulticastDelegate MulticastDelegateProp;
+
+	//
+	// Check type hinted methods.
+	//
+
+	UFUNCTION(BlueprintPure, Category = "Python|Internal")
+	bool CheckBoolTypeHints(bool bParam1, bool bParam2 = true, bool bParam3 = false);
+
+	UFUNCTION(BlueprintPure, Category = "Python|Internal")
+	int32 CheckIntegerTypeHints(uint8 Param1, int32 Param2 = 4, int64 Param3 = 5);
+
+	UFUNCTION(BlueprintPure, Category = "Python|Internal")
+	double CheckFloatTypeHints(float Param1, double Param2, float Param3 = -3.3f, double Param4 = 4.4);
+
+	UFUNCTION(BlueprintPure, Category = "Python|Internal")
+	EPyTestEnum CheckEnumTypeHints(EPyTestEnum Param1, EPyTestEnum Param2 = EPyTestEnum::One);
+
+	UFUNCTION(BlueprintPure, Category = "Python|Internal")
+	FString CheckStringTypeHints(const FString& Param1, const FString& Param2 = TEXT("Hi"));
+
+	UFUNCTION(BlueprintPure, Category = "Python|Internal")
+	FName CheckNameTypeHints(const FName& Param1, const FName& Param2 = FName(TEXT("Hi")));
+
+	UFUNCTION(BlueprintPure, Category = "Python|Internal")
+	FText CheckTextTypeHints(const FText& Param1, const FText& Param2 = INVTEXT("Hi"));
+
+	UFUNCTION(BlueprintPure, Category = "Python|Internal")
+	TFieldPath<FProperty> CheckFieldPathTypeHints(const TFieldPath<FProperty> Param1);
+
+	UFUNCTION(BlueprintPure, Category = "Python|Internal")
+	FPyTestStruct CheckStructTypeHints(const FPyTestStruct& Param1, const FPyTestStruct& Param2 = FPyTestStruct());
+
+	UFUNCTION(BlueprintPure, Category = "Python|Internal")
+	UPyTestObject* CheckObjectTypeHints(const UPyTestObject* Param1, const UPyTestObject* Param4 = nullptr);
+
+	UFUNCTION(BlueprintPure, Category = "Python|Internal")
+	TArray<FText> CheckArrayTypeHints(const TArray<FString>& Param1, const TArray<FName>& Param2, const TArray<FText>& Param3, const TArray<UObject*>& Param4);
+
+	UFUNCTION(BlueprintPure, Category = "Python|Internal")
+	TSet<FName> CheckSetTypeHints(const TSet<FString>& Param1, const TSet<FName>& Param2, const TSet<UObject*>& Param3);
+
+	UFUNCTION(BlueprintPure, Category = "Python|Internal")
+	TMap<FString, UObject*> CheckMapTypeHints(const TMap<int, FString>& Param1, const TMap<int, FName>& Param2, const TMap<int, FText>& Param3, const TMap<int, UObject*>& Param4);
+
+	UFUNCTION(BlueprintPure, Category = "Python|Internal")
+	FPyTestDelegate& CheckDelegateTypeHints(const FPyTestDelegate& Param1);
+
+	UFUNCTION(BlueprintPure, Category = "Python|Internal")
+	static bool CheckStaticFunction(bool Param1, int32 Param2, double Param3, const FString& Param4);
+
+	UFUNCTION(BlueprintPure, Category = "Python|Internal")
+	static int CheckTupleReturnType(UPARAM(ref) FString& InOutString);
+
+	//
+	// Members to facilitate testing particular Python API.
+	//
+
+	UPROPERTY(EditAnywhere, BlueprintReadWrite, Category = "Python|Internal")
+	FPyTestSlateTickDelegate SlateTickDelegate;
+};
+
diff --git a/Engine/Plugins/Experimental/PythonScriptPlugin/Source/PythonScriptPlugin/Private/PyWrapperArray.cpp b/Engine/Plugins/Experimental/PythonScriptPlugin/Source/PythonScriptPlugin/Private/PyWrapperArray.cpp
--- a/Engine/Plugins/Experimental/PythonScriptPlugin/Source/PythonScriptPlugin/Private/PyWrapperArray.cpp
+++ b/Engine/Plugins/Experimental/PythonScriptPlugin/Source/PythonScriptPlugin/Private/PyWrapperArray.cpp
@@ -1374,7 +1374,7 @@ PyTypeObject InitializePyWrapperArrayType()
 			}
 
 			int32 Len = 0;
-			if (!PyUtil::ValidateContainerLenParam(PyObj, Len, "len", *PyUtil::GetErrorContext(InSelf)))
+			if (PyUtil::ValidateContainerLenParam(PyObj, Len, "len", *PyUtil::GetErrorContext(InSelf)) != 0)
 			{
 				return nullptr;
 			}
@@ -1388,20 +1388,21 @@ PyTypeObject InitializePyWrapperArrayType()
 		}
 	};
 
+	// NOTE: _ElemType = typing.TypeVar('_ElemType') and Type/Optional/Iterable/Callable are defines by the Python typing module.
 	static PyMethodDef PyMethods[] = {
-		{ "cast", PyCFunctionCast(&FMethods::Cast), METH_VARARGS | METH_KEYWORDS | METH_CLASS, "X.cast(type, obj) -> Array -- cast the given object to this Unreal array type" },
-		{ "__copy__", PyCFunctionCast(&FMethods::Copy), METH_NOARGS, "x.__copy__() -> Array -- copy this Unreal array" },
-		{ "copy", PyCFunctionCast(&FMethods::Copy), METH_NOARGS, "x.copy() -> Array -- copy this Unreal array" },
-		{ "append", PyCFunctionCast(&FMethods::Append), METH_VARARGS, "x.append(value) -> None -- append the given value to the end of this Unreal array (equivalent to TArray::Add in C++)" },
-		{ "count", PyCFunctionCast(&FMethods::Count), METH_VARARGS, "x.count(value) -> integer -- return the number of times that value appears in this this Unreal array" },
-		{ "extend", PyCFunctionCast(&FMethods::Extend), METH_VARARGS, "x.extend(iterable) -> None -- extend this Unreal array by appending elements from the given iterable (equivalent to TArray::Append in C++)" },
-		{ "index", PyCFunctionCast(&FMethods::Index), METH_VARARGS | METH_KEYWORDS, "x.index(value, start=0, stop=len) -> integer -- get the index of the first matching value in this Unreal array, or raise ValueError if missing (equivalent to TArray::IndexOfByKey in C++)" },
-		{ "insert", PyCFunctionCast(&FMethods::Insert), METH_VARARGS | METH_KEYWORDS, "x.insert(index, value) -> None -- insert the given value at the given index in this Unreal array" },
-		{ "pop", PyCFunctionCast(&FMethods::Pop), METH_VARARGS, "x.pop(index=len-1) -> value -- remove and return the value at the given index in this Unreal array, or raise IndexError if the index is out-of-bounds" },
-		{ "remove", PyCFunctionCast(&FMethods::Remove), METH_VARARGS, "x.remove(value) -> None -- remove the first matching value in this Unreal array, or raise ValueError if missing" },
-		{ "reverse", PyCFunctionCast(&FMethods::Reverse), METH_NOARGS, "x.reverse() -> None -- reverse this Unreal array in-place" },
-		{ "sort", PyCFunctionCast(&FMethods::Sort), METH_VARARGS | METH_KEYWORDS, "x.sort(key=None, reverse=False) -> None -- stable sort this Unreal array in-place" },
-		{ "resize", PyCFunctionCast(&FMethods::Resize), METH_VARARGS, "x.resize(len) -> None -- resize this Unreal array to hold the given number of elements" },
+		{ "cast", PyCFunctionCast(&FMethods::Cast), METH_VARARGS | METH_KEYWORDS | METH_CLASS, "cast(cls, type: Type[_ElemType], obj: object) -> Array[_ElemType] -- cast the given object to this Unreal array type" },
+		{ "__copy__", PyCFunctionCast(&FMethods::Copy), METH_NOARGS, "__copy__(self) -> Array[_ElemType] -- copy this Unreal array" },
+		{ "copy", PyCFunctionCast(&FMethods::Copy), METH_NOARGS, "copy(self) -> Array[_ElemType] -- copy this Unreal array" },
+		{ "append", PyCFunctionCast(&FMethods::Append), METH_VARARGS, "append(self, value: _ElemType) -> None -- append the given value to the end of this Unreal array (equivalent to TArray::Add in C++)" },
+		{ "count", PyCFunctionCast(&FMethods::Count), METH_VARARGS, "count(self, value: _ElemType) -> int -- return the number of times that value appears in this this Unreal array" },
+		{ "extend", PyCFunctionCast(&FMethods::Extend), METH_VARARGS, "extend(self, values: Iterable[_ElemType]) -> None -- extend this Unreal array by appending elements from the given iterable (equivalent to TArray::Append in C++)" },
+		{ "index", PyCFunctionCast(&FMethods::Index), METH_VARARGS | METH_KEYWORDS, "index(self, value: _ElemType, start: int = 0, stop: int = -1) -> int -- get the index of the first matching value in this Unreal array, or raise ValueError if missing (equivalent to TArray::IndexOfByKey in C++)" },
+		{ "insert", PyCFunctionCast(&FMethods::Insert), METH_VARARGS | METH_KEYWORDS, "insert(self, index: int, value: _ElemType) -> None -- insert the given value at the given index in this Unreal array" },
+		{ "pop", PyCFunctionCast(&FMethods::Pop), METH_VARARGS, "pop(self, index: int = -1) -> _ElemType -- remove and return the value at the given index in this Unreal array, or raise IndexError if the index is out-of-bounds" },
+		{ "remove", PyCFunctionCast(&FMethods::Remove), METH_VARARGS, "remove(self, value: _ElemType) -> None -- remove the first matching value in this Unreal array, or raise ValueError if missing" },
+		{ "reverse", PyCFunctionCast(&FMethods::Reverse), METH_NOARGS, "reverse(self) -> None -- reverse this Unreal array in-place" },
+		{ "sort", PyCFunctionCast(&FMethods::Sort), METH_VARARGS | METH_KEYWORDS, "sort(self, key: Optional[Callable[[_ElemType], object]]=None, reverse: bool=False) -> None -- stable sort this Unreal array in-place" },
+		{ "resize", PyCFunctionCast(&FMethods::Resize), METH_VARARGS, "resize(self, len: int) -> None -- resize this Unreal array to hold the given number of elements" },
 		{ nullptr, nullptr, 0, nullptr }
 	};
 
diff --git a/Engine/Plugins/Experimental/PythonScriptPlugin/Source/PythonScriptPlugin/Private/PyWrapperDelegate.cpp b/Engine/Plugins/Experimental/PythonScriptPlugin/Source/PythonScriptPlugin/Private/PyWrapperDelegate.cpp
--- a/Engine/Plugins/Experimental/PythonScriptPlugin/Source/PythonScriptPlugin/Private/PyWrapperDelegate.cpp
+++ b/Engine/Plugins/Experimental/PythonScriptPlugin/Source/PythonScriptPlugin/Private/PyWrapperDelegate.cpp
@@ -771,17 +771,20 @@ PyTypeObject InitializePyWrapperDelegateType()
 		}
 	};
 
+	// NOTE: The delegate couldn't be strictly type-hinted. The exact return/param types of the delegate are unknown here. In Python 3.11, we should be able
+	//       to use typing.Self for return type of __copy__() and copy() methods.
+	//       _T = typing.TypeVar('_T') and Type/Any/Callable/Union are defines by the Python typing module.
 	static PyMethodDef PyMethods[] = {
-		{ "cast", PyCFunctionCast(&FMethods::Cast), METH_VARARGS | METH_CLASS, "X.cast(object) -> struct -- cast the given object to this Unreal delegate type" },
-		{ "__copy__", PyCFunctionCast(&FMethods::Copy), METH_NOARGS, "x.__copy__() -> delegate -- copy this Unreal delegate" },
-		{ "copy", PyCFunctionCast(&FMethods::Copy), METH_NOARGS, "x.copy() -> struct -- copy this Unreal delegate" },
-		{ "is_bound", PyCFunctionCast(&FMethods::IsBound), METH_NOARGS, "x.is_bound() -> bool -- is this Unreal delegate bound to something?" },
-		{ "bind_delegate", PyCFunctionCast(&FMethods::BindDelegate), METH_VARARGS, "x.bind_delegate(delegate) -> None -- bind this Unreal delegate to the same object and function as another delegate" },
-		{ "bind_function", PyCFunctionCast(&FMethods::BindFunction), METH_VARARGS, "x.bind_function(obj, name) -> None -- bind this Unreal delegate to a named Unreal function on the given object instance" },
-		{ "bind_callable", PyCFunctionCast(&FMethods::BindCallable), METH_VARARGS, "x.bind_callable(callable) -> None -- bind this Unreal delegate to a Python callable" },
-		{ "unbind", PyCFunctionCast(&FMethods::Unbind), METH_NOARGS, "x.unbind() -> None -- unbind this Unreal delegate" },
-		{ "execute", PyCFunctionCast(&FMethods::Execute), METH_VARARGS, "x.execute(...) -> value -- call this Unreal delegate, but error if it's unbound" },
-		{ "execute_if_bound", PyCFunctionCast(&FMethods::ExecuteIfBound), METH_VARARGS, "x.execute_if_bound(...) -> value -- call this Unreal delegate, but only if it's bound to something" },
+		{ "cast", PyCFunctionCast(&FMethods::Cast), METH_VARARGS | METH_CLASS, "cast(cls: Type[_T], object: object) -> _T -- cast the given object to this Unreal delegate type" },
+		{ "__copy__", PyCFunctionCast(&FMethods::Copy), METH_NOARGS, "__copy__(self) -> Any -- copy this Unreal delegate" },
+		{ "copy", PyCFunctionCast(&FMethods::Copy), METH_NOARGS, "copy(self) -> Any -- copy this Unreal delegate" },
+		{ "is_bound", PyCFunctionCast(&FMethods::IsBound), METH_NOARGS, "is_bound(self) -> bool -- is this Unreal delegate bound to something?" },
+		{ "bind_delegate", PyCFunctionCast(&FMethods::BindDelegate), METH_VARARGS, "bind_delegate(self, delegate: DelegateBase) -> None -- bind this Unreal delegate to the same object and function as another delegate" },
+		{ "bind_function", PyCFunctionCast(&FMethods::BindFunction), METH_VARARGS, "bind_function(self, obj: Object, name: Union[Name, str]) -> None -- bind this Unreal delegate to a named Unreal function on the given object instance" },
+		{ "bind_callable", PyCFunctionCast(&FMethods::BindCallable), METH_VARARGS, "bind_callable(self, callable: Callable[..., Any]) -> None -- bind this Unreal delegate to a Python callable" },
+		{ "unbind", PyCFunctionCast(&FMethods::Unbind), METH_NOARGS, "unbind(self) -> None -- unbind this Unreal delegate" },
+		{ "execute", PyCFunctionCast(&FMethods::Execute), METH_VARARGS, "execute(self, *args: Any) -> Any -- call this Unreal delegate, but error if it's unbound" },
+		{ "execute_if_bound", PyCFunctionCast(&FMethods::ExecuteIfBound), METH_VARARGS, "x.execute_if_bound(self, *args: Any) -> Any -- call this Unreal delegate, but only if it's bound to something" },
 		{ nullptr, nullptr, 0, nullptr }
 	};
 
@@ -1189,23 +1192,25 @@ PyTypeObject InitializePyWrapperMulticastDelegateType()
 			return FPyWrapperMulticastDelegate::CallDelegate(InSelf, InArgs);
 		}
 	};
-
+	// NOTE: The multicast delegate couldn't be strictly type-hinted. The exact param types of the delegate are unknown here. In Python 3.11, we should be able
+	//       to use typing.Self for return type of __copy__() and copy() methods.
+	//       _T = typing.TypeVar('_T') and Type/Any/Callable/Union are defines by the Python typing module.
 	static PyMethodDef PyMethods[] = {
-		{ "cast", PyCFunctionCast(&FMethods::Cast), METH_VARARGS | METH_CLASS, "X.cast(object) -> struct -- cast the given object to this Unreal delegate type" },
-		{ "__copy__", PyCFunctionCast(&FMethods::Copy), METH_NOARGS, "x.__copy__() -> struct -- copy this Unreal delegate" },
-		{ "copy", PyCFunctionCast(&FMethods::Copy), METH_NOARGS, "x.copy() -> struct -- copy this Unreal delegate" },
-		{ "is_bound", PyCFunctionCast(&FMethods::IsBound), METH_NOARGS, "x.is_bound() -> bool -- is this Unreal delegate bound to something?" },
-		{ "add_function", PyCFunctionCast(&FMethods::AddFunction), METH_VARARGS, "x.add_function(obj, name) -> None -- add a binding to a named Unreal function on the given object instance to the invocation list of this Unreal delegate" },
-		{ "add_callable", PyCFunctionCast(&FMethods::AddCallable), METH_VARARGS, "x.add_callable(callable) -> None -- add a binding to a Python callable to the invocation list of this Unreal delegate" },
-		{ "add_function_unique", PyCFunctionCast(&FMethods::AddFunctionUnique), METH_VARARGS, "x.add_function_unique(obj, name) -> None -- add a unique binding to a named Unreal function on the given object instance to the invocation list of this Unreal delegate" },
-		{ "add_callable_unique", PyCFunctionCast(&FMethods::AddCallableUnique), METH_VARARGS, "x.add_callable_unique(callable) -> None -- add a unique binding to a Python callable to the invocation list of this Unreal delegate" },
-		{ "remove_function", PyCFunctionCast(&FMethods::RemoveFunction), METH_VARARGS, "x.remove_function(obj, name) -> None -- remove a binding to a named Unreal function on the given object instance from the invocation list of this Unreal delegate" },
-		{ "remove_callable", PyCFunctionCast(&FMethods::RemoveCallable), METH_VARARGS, "x.remove_callable(callable) -> None -- remove a binding to a Python callable from the invocation list of this Unreal delegate" },
-		{ "remove_object", PyCFunctionCast(&FMethods::RemoveObject), METH_VARARGS, "x.remove_object(obj) -> None -- remove all bindings for the given object instance from the invocation list of this Unreal delegate" },
-		{ "contains_function", PyCFunctionCast(&FMethods::ContainsFunction), METH_VARARGS, "x.contains_function(obj, name) -> bool -- does the invocation list of this Unreal delegate contain a binding to a named Unreal function on the given object instance" },
-		{ "contains_callable", PyCFunctionCast(&FMethods::ContainsCallable), METH_VARARGS, "x.contains_callable(callable) -> bool -- does the invocation list of this Unreal delegate contain a binding to a Python callable" },
+		{ "cast", PyCFunctionCast(&FMethods::Cast), METH_VARARGS | METH_CLASS, "cast(cls: Type[_T], object: object) -> _T -- cast the given object to this Unreal delegate type" },
+		{ "__copy__", PyCFunctionCast(&FMethods::Copy), METH_NOARGS, "__copy__(self) -> Any -- copy this Unreal delegate" },
+		{ "copy", PyCFunctionCast(&FMethods::Copy), METH_NOARGS, "copy(self) -> Any -- copy this Unreal delegate" },
+		{ "is_bound", PyCFunctionCast(&FMethods::IsBound), METH_NOARGS, "is_bound(self) -> bool -- is this Unreal delegate bound to something?" },
+		{ "add_function", PyCFunctionCast(&FMethods::AddFunction), METH_VARARGS, "add_function(self, obj: Object, name: Union[Name, str]) -> None -- add a binding to a named Unreal function on the given object instance to the invocation list of this Unreal delegate" },
+		{ "add_callable", PyCFunctionCast(&FMethods::AddCallable), METH_VARARGS, "add_callable(self, callable: Callable[..., Any]) -> None -- add a binding to a Python callable to the invocation list of this Unreal delegate" },
+		{ "add_function_unique", PyCFunctionCast(&FMethods::AddFunctionUnique), METH_VARARGS, "add_function_unique(self, obj: Object, name: Union[Name, str]) -> None -- add a unique binding to a named Unreal function on the given object instance to the invocation list of this Unreal delegate" },
+		{ "add_callable_unique", PyCFunctionCast(&FMethods::AddCallableUnique), METH_VARARGS, "add_callable_unique(self, callable: Callable[..., Any]) -> None -- add a unique binding to a Python callable to the invocation list of this Unreal delegate" },
+		{ "remove_function", PyCFunctionCast(&FMethods::RemoveFunction), METH_VARARGS, "remove_function(self, obj: Object, name: Union[Name, str]) -> None -- remove a binding to a named Unreal function on the given object instance from the invocation list of this Unreal delegate" },
+		{ "remove_callable", PyCFunctionCast(&FMethods::RemoveCallable), METH_VARARGS, "remove_callable(self, callable: Callable[..., Any]) -> None -- remove a binding to a Python callable from the invocation list of this Unreal delegate" },
+		{ "remove_object", PyCFunctionCast(&FMethods::RemoveObject), METH_VARARGS, "remove_object(self, obj: Object) -> None -- remove all bindings for the given object instance from the invocation list of this Unreal delegate" },
+		{ "contains_function", PyCFunctionCast(&FMethods::ContainsFunction), METH_VARARGS, "contains_function(self, obj: Object, name: Union[Name, str]) -> bool -- does the invocation list of this Unreal delegate contain a binding to a named Unreal function on the given object instance" },
+		{ "contains_callable", PyCFunctionCast(&FMethods::ContainsCallable), METH_VARARGS, "contains_callable(self, callable: Callable[..., Any]) -> bool -- does the invocation list of this Unreal delegate contain a binding to a Python callable" },
 		{ "clear", PyCFunctionCast(&FMethods::Clear), METH_NOARGS, "x.clear() -> None -- clear the invocation list of this Unreal delegate" },
-		{ "broadcast", PyCFunctionCast(&FMethods::Broadcast), METH_VARARGS, "x.broadcast(...) -> None -- invoke this Unreal multicast delegate" },
+		{ "broadcast", PyCFunctionCast(&FMethods::Broadcast), METH_VARARGS, "x.broadcast(*args: Any) -> None -- invoke this Unreal multicast delegate" },
 		{ nullptr, nullptr, 0, nullptr }
 	};
 
diff --git a/Engine/Plugins/Experimental/PythonScriptPlugin/Source/PythonScriptPlugin/Private/PyWrapperEnum.cpp b/Engine/Plugins/Experimental/PythonScriptPlugin/Source/PythonScriptPlugin/Private/PyWrapperEnum.cpp
--- a/Engine/Plugins/Experimental/PythonScriptPlugin/Source/PythonScriptPlugin/Private/PyWrapperEnum.cpp
+++ b/Engine/Plugins/Experimental/PythonScriptPlugin/Source/PythonScriptPlugin/Private/PyWrapperEnum.cpp
@@ -423,10 +423,11 @@ PyTypeObject InitializePyWrapperEnumType()
 		{ nullptr, 0, 0, 0, nullptr }
 	};
 
+	// NOTE: _T is defines as: _T = typing.TypeVar('_T') and Type is defined by the Python typing module.
 	static PyMethodDef PyMethods[] = {
-		{ "cast", PyCFunctionCast(&FMethods::Cast), METH_VARARGS | METH_CLASS, "X.cast(object) -> enum -- cast the given object to this Unreal enum type" },
-		{ "static_enum", PyCFunctionCast(&FMethods::StaticEnum), METH_NOARGS | METH_CLASS, "X.static_enum() -> Enum -- get the Unreal enum of this type" },
-		{ "get_display_name", PyCFunctionCast(&FMethods::GetDisplayName), METH_NOARGS, "X.get_display_name() -> Text -- get the UMETA display name of this type in the current culture" },
+		{ "cast", PyCFunctionCast(&FMethods::Cast), METH_VARARGS | METH_CLASS, "cast(cls: Type[_T], object: object) -> _T -- cast the given object to this Unreal enum type" },
+		{ "static_enum", PyCFunctionCast(&FMethods::StaticEnum), METH_NOARGS | METH_CLASS, "static_enum(cls) -> Enum -- get the Unreal enum of this type" },
+		{ "get_display_name", PyCFunctionCast(&FMethods::GetDisplayName), METH_NOARGS, "get_display_name(self) -> Text -- get the UMETA display name of this type in the current culture" },
 		{ nullptr, nullptr, 0, nullptr }
 	};
 
diff --git a/Engine/Plugins/Experimental/PythonScriptPlugin/Source/PythonScriptPlugin/Private/PyWrapperFieldPath.cpp b/Engine/Plugins/Experimental/PythonScriptPlugin/Source/PythonScriptPlugin/Private/PyWrapperFieldPath.cpp
--- a/Engine/Plugins/Experimental/PythonScriptPlugin/Source/PythonScriptPlugin/Private/PyWrapperFieldPath.cpp
+++ b/Engine/Plugins/Experimental/PythonScriptPlugin/Source/PythonScriptPlugin/Private/PyWrapperFieldPath.cpp
@@ -191,10 +191,10 @@ PyTypeObject InitializePyWrapperFieldPathType()
 	};
 
 	static PyMethodDef PyMethods[] = {
-		{ "cast", PyCFunctionCast(&FMethods::Cast), METH_VARARGS | METH_CLASS, "X.cast(obj: Any) -> FieldPath -- cast the given object to this Unreal field path type" },
-		{ "__copy__", PyCFunctionCast(&FMethods::Copy), METH_NOARGS, "x.__copy__() -> FieldPath -- copy this Unreal field path" },
-		{ "copy", PyCFunctionCast(&FMethods::Copy), METH_NOARGS, "x.copy() -> FieldPath -- copy this Unreal field path" },
-		{ "is_valid", PyCFunctionCast(&FMethods::IsValid), METH_NOARGS, "x.is_valid() -> bool -- whether this Unreal field path refers to an existing Unreal field" },
+		{ "cast", PyCFunctionCast(&FMethods::Cast), METH_VARARGS | METH_CLASS, "cast(cls, obj: object) -> FieldPath -- cast the given object to this Unreal field path type" },
+		{ "__copy__", PyCFunctionCast(&FMethods::Copy), METH_NOARGS, "__copy__(self) -> FieldPath -- copy this Unreal field path" },
+		{ "copy", PyCFunctionCast(&FMethods::Copy), METH_NOARGS, "copy(self) -> FieldPath -- copy this Unreal field path" },
+		{ "is_valid", PyCFunctionCast(&FMethods::IsValid), METH_NOARGS, "is_valid(self) -> bool -- whether this Unreal field path refers to an existing Unreal field" },
 		{ nullptr, nullptr, 0, nullptr }
 	};
 
diff --git a/Engine/Plugins/Experimental/PythonScriptPlugin/Source/PythonScriptPlugin/Private/PyWrapperFixedArray.cpp b/Engine/Plugins/Experimental/PythonScriptPlugin/Source/PythonScriptPlugin/Private/PyWrapperFixedArray.cpp
--- a/Engine/Plugins/Experimental/PythonScriptPlugin/Source/PythonScriptPlugin/Private/PyWrapperFixedArray.cpp
+++ b/Engine/Plugins/Experimental/PythonScriptPlugin/Source/PythonScriptPlugin/Private/PyWrapperFixedArray.cpp
@@ -794,10 +794,11 @@ PyTypeObject InitializePyWrapperFixedArrayType()
 		}
 	};
 
+	// NOTE: _ElemType = typing.TypeVar('_ElemType') and TYpe is defined in the typing Python module.
 	static PyMethodDef PyMethods[] = {
-		{ "cast", PyCFunctionCast(&FMethods::Cast), METH_VARARGS | METH_KEYWORDS | METH_CLASS, "X.cast(type, obj) -> FixedArray -- cast the given object to this Unreal fixed-array type" },
-		{ "__copy__", PyCFunctionCast(&FMethods::Copy), METH_NOARGS, "x.__copy__() -> FixedArray -- copy this Unreal fixed-array" },
-		{ "copy", PyCFunctionCast(&FMethods::Copy), METH_NOARGS, "x.copy() -> FixedArray -- copy this Unreal fixed-array" },
+		{ "cast", PyCFunctionCast(&FMethods::Cast), METH_VARARGS | METH_KEYWORDS | METH_CLASS, "cast(cls, type: Type[_ElemType], obj: object) -> FixedArray[_ElemType] -- cast the given object to this Unreal fixed-array type" },
+		{ "__copy__", PyCFunctionCast(&FMethods::Copy), METH_NOARGS, "__copy__(self) -> FixedArray[_ElemType] -- copy this Unreal fixed-array" },
+		{ "copy", PyCFunctionCast(&FMethods::Copy), METH_NOARGS, "copy(self) -> FixedArray[_ElemType] -- copy this Unreal fixed-array" },
 		{ nullptr, nullptr, 0, nullptr }
 	};
 
diff --git a/Engine/Plugins/Experimental/PythonScriptPlugin/Source/PythonScriptPlugin/Private/PyWrapperMap.cpp b/Engine/Plugins/Experimental/PythonScriptPlugin/Source/PythonScriptPlugin/Private/PyWrapperMap.cpp
--- a/Engine/Plugins/Experimental/PythonScriptPlugin/Source/PythonScriptPlugin/Private/PyWrapperMap.cpp
+++ b/Engine/Plugins/Experimental/PythonScriptPlugin/Source/PythonScriptPlugin/Private/PyWrapperMap.cpp
@@ -1490,20 +1490,21 @@ PyTypeObject InitializePyWrapperMapType()
 		}
 	};
 
+	// NOTE: _KeyType = typing.TypeVar('_KeyType'), _ValueType = typing.TypeVar('_ValueType'), Type/Iterable/Tuple/Union/List/ItemsView/KeysView/ValuesView are defines in the Python typing module.
 	static PyMethodDef PyMethods[] = {
-		{ "cast", PyCFunctionCast(&FMethods::Cast), METH_VARARGS | METH_KEYWORDS | METH_CLASS, "X.cast(key, value, obj) -> Map -- cast the given object to this Unreal map type" },
-		{ "__copy__", PyCFunctionCast(&FMethods::Copy), METH_NOARGS, "x.__copy__() -> Map -- copy this Unreal map" },
-		{ "copy", PyCFunctionCast(&FMethods::Copy), METH_NOARGS, "x.copy() -> Map -- copy this Unreal map" },
-		{ "clear", PyCFunctionCast(&FMethods::Clear), METH_NOARGS, "x.clear() -> None -- remove all items from this Unreal map" },
-		{ "fromkeys", PyCFunctionCast(&FMethods::FromKeys), METH_VARARGS | METH_KEYWORDS | METH_CLASS, "X.fromkeys(sequence, value=None) -> Map -- returns a new Unreal map of keys from the sequence using the given value (types are inferred)" },
-		{ "get", PyCFunctionCast(&FMethods::Get), METH_VARARGS | METH_KEYWORDS, "x.get(key, default=None) -> value -- x[key] if key in x, otherwise default" },
-		{ "setdefault", PyCFunctionCast(&FMethods::SetDefault), METH_VARARGS | METH_KEYWORDS, "x.setdefault(key, default=None) -> value -- set key to default if key not in x and return the new value of key" },
-		{ "pop", PyCFunctionCast(&FMethods::Pop), METH_VARARGS | METH_KEYWORDS, "x.pop(key, default=None) -> value -- remove key and return its value, or default if key not present, or raise KeyError if no default" },
-		{ "popitem", PyCFunctionCast(&FMethods::PopItem), METH_NOARGS, "x.popitem() -> pair -- remove and return an arbitrary pair from this Unreal map, or raise KeyError if the map is empty" },
-		{ "update", PyCFunctionCast(&FMethods::Update), METH_VARARGS | METH_KEYWORDS, "x.update(...) -> None -- update this Unreal map from the given mapping or sequence pairs type or key->value arguments" },
-		{ "items", PyCFunctionCast(&FMethods::Items), METH_NOARGS, "x.items() -> view -- a set-like view of the key->value pairs of this Unreal map" },
-		{ "keys", PyCFunctionCast(&FMethods::Keys), METH_NOARGS, "x.keys() -> view -- a set-like view of the keys of this Unreal map" },
-		{ "values", PyCFunctionCast(&FMethods::Values), METH_NOARGS, "x.values() -> view -- a view of the values of this Unreal map" },
+		{ "cast", PyCFunctionCast(&FMethods::Cast), METH_VARARGS | METH_KEYWORDS | METH_CLASS, "cast(cls, key: Type[_KeyType], value: Type[_ValueType], obj: object) -> Map[_KeyType, _ValueType] -- cast the given object to this Unreal map type" },
+		{ "__copy__", PyCFunctionCast(&FMethods::Copy), METH_NOARGS, "__copy__(self) -> Map[_KeyType, _ValueType] -- copy this Unreal map" },
+		{ "copy", PyCFunctionCast(&FMethods::Copy), METH_NOARGS, "copy(self) -> Map[_KeyType, _ValueType] -- copy this Unreal map" },
+		{ "clear", PyCFunctionCast(&FMethods::Clear), METH_NOARGS, "clear(self) -> None -- remove all items from this Unreal map" },
+		{ "fromkeys", PyCFunctionCast(&FMethods::FromKeys), METH_VARARGS | METH_KEYWORDS | METH_CLASS, "fromkeys(cls, sequence: Iterable[_KeyType], value: _ValueType) -> Map[_KeyType, _ValueType] -- returns a new Unreal map of keys from the sequence using the given value (types are inferred)" },
+		{ "get", PyCFunctionCast(&FMethods::Get), METH_VARARGS | METH_KEYWORDS, "get(self, key: _KeyType, default: _ValueType=...) -> _ValueType -- x[key] if key in x, otherwise default" },
+		{ "setdefault", PyCFunctionCast(&FMethods::SetDefault), METH_VARARGS | METH_KEYWORDS, "setdefault(self, key: _KeyType, default: _ValueType=...) -> _ValueType -- set key to default if key not in x and return the value of key" },
+		{ "pop", PyCFunctionCast(&FMethods::Pop), METH_VARARGS | METH_KEYWORDS, "pop(self, key: _KeyType, default: _ValueType=...) -> _ValueType -- remove key and return its value, or default if key not present, or raise KeyError if no default" },
+		{ "popitem", PyCFunctionCast(&FMethods::PopItem), METH_NOARGS, "popitem(self) -> Tuple[_KeyType, _ValueType] -- remove and return an arbitrary pair from this Unreal map, or raise KeyError if the map is empty" },
+		{ "update", PyCFunctionCast(&FMethods::Update), METH_VARARGS | METH_KEYWORDS, "update(self, pairs: Union[Mapping[_KeyType, _ValueType], Iterable[Tuple[_KeyType, _ValueType]], Iterable[List[Union[_KeyType, _ValueType]]]]) -> None -- update this Unreal map from the given mapping or sequence pairs type or key->value arguments" },
+		{ "items", PyCFunctionCast(&FMethods::Items), METH_NOARGS, "items(self) -> ItemsView[_KeyType, _ValueType] -- a set-like view of the key->value pairs of this Unreal map" },
+		{ "keys", PyCFunctionCast(&FMethods::Keys), METH_NOARGS, "keys(self) -> KeysView[_KeyType] -- a set-like view of the keys of this Unreal map" },
+		{ "values", PyCFunctionCast(&FMethods::Values), METH_NOARGS, "values(self) -> ValuesView[_ValueType] -- a view of the values of this Unreal map" },
 		{ nullptr, nullptr, 0, nullptr }
 	};
 
diff --git a/Engine/Plugins/Experimental/PythonScriptPlugin/Source/PythonScriptPlugin/Private/PyWrapperName.cpp b/Engine/Plugins/Experimental/PythonScriptPlugin/Source/PythonScriptPlugin/Private/PyWrapperName.cpp
--- a/Engine/Plugins/Experimental/PythonScriptPlugin/Source/PythonScriptPlugin/Private/PyWrapperName.cpp
+++ b/Engine/Plugins/Experimental/PythonScriptPlugin/Source/PythonScriptPlugin/Private/PyWrapperName.cpp
@@ -166,9 +166,9 @@ PyTypeObject InitializePyWrapperNameType()
 	};
 
 	static PyMethodDef PyMethods[] = {
-		{ "cast", PyCFunctionCast(&FMethods::Cast), METH_VARARGS | METH_CLASS, "X.cast(object) -> Name -- cast the given object to this Unreal name type" },
-		{ "is_valid", PyCFunctionCast(&FMethods::IsValid), METH_NOARGS, "x.is_valid() -> bool -- is this Unreal name valid?" },
-		{ "is_none", PyCFunctionCast(&FMethods::IsNone), METH_NOARGS, "x.is_none() -> bool -- is this Unreal name set to NAME_None?" },
+		{ "cast", PyCFunctionCast(&FMethods::Cast), METH_VARARGS | METH_CLASS, "cast(cls, object: object) -> Name -- cast the given object to this Unreal name type" },
+		{ "is_valid", PyCFunctionCast(&FMethods::IsValid), METH_NOARGS, "is_valid(self) -> bool -- is this Unreal name valid?" },
+		{ "is_none", PyCFunctionCast(&FMethods::IsNone), METH_NOARGS, "is_none(self) -> bool -- is this Unreal name set to NAME_None?" },
 		{ nullptr, nullptr, 0, nullptr }
 	};
 
diff --git a/Engine/Plugins/Experimental/PythonScriptPlugin/Source/PythonScriptPlugin/Private/PyWrapperObject.cpp b/Engine/Plugins/Experimental/PythonScriptPlugin/Source/PythonScriptPlugin/Private/PyWrapperObject.cpp
--- a/Engine/Plugins/Experimental/PythonScriptPlugin/Source/PythonScriptPlugin/Private/PyWrapperObject.cpp
+++ b/Engine/Plugins/Experimental/PythonScriptPlugin/Source/PythonScriptPlugin/Private/PyWrapperObject.cpp
@@ -992,28 +992,29 @@ PyTypeObject InitializePyWrapperObjectType()
 		}
 	};
 
+	// NOTE: _T = typing.TypeVar('_T') and Any/Type/Union/Mapping/Optional are defines by the Python typing module.
 	static PyMethodDef PyMethods[] = {
-		{ PyGenUtil::PostInitFuncName, PyCFunctionCast(&FMethods::PostInit), METH_NOARGS, "x._post_init() -> None -- called during Unreal object initialization (equivalent to PostInitProperties in C++)" },
-		{ "cast", PyCFunctionCast(&FMethods::Cast), METH_VARARGS | METH_CLASS, "X.cast(object) -> Object -- cast the given object to this Unreal object type" },
-		{ "get_default_object", PyCFunctionCast(&FMethods::GetDefaultObject), METH_NOARGS | METH_CLASS, "X.get_default_object() -> Object -- get the Unreal class default object (CDO) of this type" },
-		{ "static_class", PyCFunctionCast(&FMethods::StaticClass), METH_NOARGS | METH_CLASS, "X.static_class() -> Class -- get the Unreal class of this type" },
-		{ "get_class", PyCFunctionCast(&FMethods::GetClass), METH_NOARGS, "x.get_class() -> Class -- get the Unreal class of this instance" },
-		{ "get_outer", PyCFunctionCast(&FMethods::GetOuter), METH_NOARGS, "x.get_outer() -> Object -- get the outer object from this instance (if any)" },
-		{ "get_typed_outer", PyCFunctionCast(&FMethods::GetTypedOuter), METH_VARARGS, "x.get_typed_outer(type) -> type() -- get the first outer object of the given type from this instance (if any)" },
-		{ "get_outermost", PyCFunctionCast(&FMethods::GetOutermost), METH_NOARGS, "x.get_outermost() -> Package -- get the outermost object (the package) from this instance" },
-		{ "is_package_external", PyCFunctionCast(&FMethods::IsPackageExternal), METH_NOARGS, "x.is_package_external() -> bool -- returns true if this instance has a different package than its outer's package" },
-		{ "get_package", PyCFunctionCast(&FMethods::GetPackage), METH_NOARGS, "x.get_package() -> Package -- get the package directly associated with this instance" },
-		{ "get_name", PyCFunctionCast(&FMethods::GetName), METH_NOARGS, "x.get_name() -> str -- get the name of this instance" },
-		{ "get_fname", PyCFunctionCast(&FMethods::GetFName), METH_NOARGS, "x.get_fname() -> FName -- get the name of this instance" },
-		{ "get_full_name", PyCFunctionCast(&FMethods::GetFullName), METH_NOARGS, "x.get_full_name() -> str -- get the full name (class name + full path) of this instance" },
-		{ "get_path_name", PyCFunctionCast(&FMethods::GetPathName), METH_NOARGS, "x.get_path_name() -> str -- get the path name of this instance" },
-		{ "get_world", PyCFunctionCast(&FMethods::GetWorld), METH_NOARGS, "x.get_world() -> World -- get the world associated with this instance (if any)" },
-		{ "modify", PyCFunctionCast(&FMethods::Modify), METH_VARARGS, "x.modify(bool) -> bool -- inform that this instance is about to be modified (tracks changes for undo/redo if transactional)" },
-		{ "rename", PyCFunctionCast(&FMethods::Rename), METH_VARARGS | METH_KEYWORDS, "x.rename(name=None, outer=None) -> bool -- rename this instance" },
-		{ "get_editor_property", PyCFunctionCast(&FMethods::GetEditorProperty), METH_VARARGS | METH_KEYWORDS, "x.get_editor_property(name) -> object -- get the value of any property visible to the editor" },
-		{ "set_editor_property", PyCFunctionCast(&FMethods::SetEditorProperty), METH_VARARGS | METH_KEYWORDS, "x.set_editor_property(name, value, notify_mode=PropertyAccessChangeNotifyMode.DEFAULT) -> None -- set the value of any property visible to the editor, ensuring that the pre/post change notifications are called" },
-		{ "set_editor_properties", PyCFunctionCast(&FMethods::SetEditorProperties), METH_VARARGS, "x.set_editor_properties(property_info) -> None -- set the value of any properties visible to the editor (from a name->value dict), ensuring that the pre/post change notifications are called" },
-		{ "call_method", PyCFunctionCast(&FMethods::CallMethod), METH_VARARGS | METH_KEYWORDS, "x.call_method(name, args=tuple(), kwargs=dict()) -> object -- call a method on this object via Unreal reflection using the given ordered (tuple) or named (dict) argument data - allows calling methods that don't have Python glue" },
+		{ PyGenUtil::PostInitFuncName, PyCFunctionCast(&FMethods::PostInit), METH_NOARGS, "_post_init(self) -> None -- called during Unreal object initialization (equivalent to PostInitProperties in C++)" },
+		{ "cast", PyCFunctionCast(&FMethods::Cast), METH_VARARGS | METH_CLASS, "cast(cls: Type[_T], object: object) -> _T -- cast the given object to this Unreal object type or raise an exception if the cast is not possible" },
+		{ "get_default_object", PyCFunctionCast(&FMethods::GetDefaultObject), METH_NOARGS | METH_CLASS, "get_default_object(cls: Type[_T]) -> _T -- get the Unreal class default object (CDO) of this type" },
+		{ "static_class", PyCFunctionCast(&FMethods::StaticClass), METH_NOARGS | METH_CLASS, "static_class(cls) -> Class -- get the Unreal class of this type" },
+		{ "get_class", PyCFunctionCast(&FMethods::GetClass), METH_NOARGS, "get_class(self) -> Class -- get the Unreal class of this instance" },
+		{ "get_outer", PyCFunctionCast(&FMethods::GetOuter), METH_NOARGS, "get_outer(self) -> Any -- get the outer object from this instance (if any)" },
+		{ "get_typed_outer", PyCFunctionCast(&FMethods::GetTypedOuter), METH_VARARGS, "get_typed_outer(self, type: Union[Class, type]) -> Any -- get the first outer object of the given type from this instance (if any)" },
+		{ "get_outermost", PyCFunctionCast(&FMethods::GetOutermost), METH_NOARGS, "get_outermost(self) -> Package -- get the outermost object (the package) from this instance" },
+		{ "is_package_external", PyCFunctionCast(&FMethods::IsPackageExternal), METH_NOARGS, "is_package_external(self) -> bool -- returns true if this instance has a different package than its outer's package" },
+		{ "get_package", PyCFunctionCast(&FMethods::GetPackage), METH_NOARGS, "get_package(self) -> Package -- get the package directly associated with this instance" },
+		{ "get_name", PyCFunctionCast(&FMethods::GetName), METH_NOARGS, "get_name(self) -> str -- get the name of this instance" },
+		{ "get_fname", PyCFunctionCast(&FMethods::GetFName), METH_NOARGS, "get_fname(self) -> Name -- get the name of this instance" },
+		{ "get_full_name", PyCFunctionCast(&FMethods::GetFullName), METH_NOARGS, "get_full_name(self) -> str -- get the full name (class name + full path) of this instance" },
+		{ "get_path_name", PyCFunctionCast(&FMethods::GetPathName), METH_NOARGS, "get_path_name(self) -> str -- get the path name of this instance" },
+		{ "get_world", PyCFunctionCast(&FMethods::GetWorld), METH_NOARGS, "get_world(self) -> Optional[World] -- get the world associated with this instance (if any)" },
+		{ "modify", PyCFunctionCast(&FMethods::Modify), METH_VARARGS, "modify(self, always_mark_dirty: bool = True) -> bool -- inform that this instance is about to be modified (tracks changes for undo/redo if transactional)" },
+		{ "rename", PyCFunctionCast(&FMethods::Rename), METH_VARARGS | METH_KEYWORDS, "rename(self, name: Union[Name, str]=\"None\", outer: Optional[Object]=None) -> bool -- rename this instance and/or change its outer" },
+		{ "get_editor_property", PyCFunctionCast(&FMethods::GetEditorProperty), METH_VARARGS | METH_KEYWORDS, "get_editor_property(self, name: str) -> object -- get the value of any property visible to the editor" },
+		{ "set_editor_property", PyCFunctionCast(&FMethods::SetEditorProperty), METH_VARARGS | METH_KEYWORDS, "set_editor_property(self, name: str, value: object, notify_mode: PropertyAccessChangeNotifyMode=PropertyAccessChangeNotifyMode.DEFAULT) -> None -- set the value of any property visible to the editor, ensuring that the pre/post change notifications are called" },
+		{ "set_editor_properties", PyCFunctionCast(&FMethods::SetEditorProperties), METH_VARARGS, "set_editor_properties(self, properties: Mapping[str, object]) -> None -- set the value of any properties visible to the editor (from a name->value dict), ensuring that the pre/post change notifications are called" },
+		{ "call_method", PyCFunctionCast(&FMethods::CallMethod), METH_VARARGS | METH_KEYWORDS, "call_method(self, name: str, *args: Any, **kwargs: Mapping[str, object]) -> Any -- call a method on this object via Unreal reflection using the given ordered (tuple) or named (dict) argument data - allows calling methods that don't have Python glue" },
 		{ nullptr, nullptr, 0, nullptr }
 	};
 
diff --git a/Engine/Plugins/Experimental/PythonScriptPlugin/Source/PythonScriptPlugin/Private/PyWrapperSet.cpp b/Engine/Plugins/Experimental/PythonScriptPlugin/Source/PythonScriptPlugin/Private/PyWrapperSet.cpp
--- a/Engine/Plugins/Experimental/PythonScriptPlugin/Source/PythonScriptPlugin/Private/PyWrapperSet.cpp
+++ b/Engine/Plugins/Experimental/PythonScriptPlugin/Source/PythonScriptPlugin/Private/PyWrapperSet.cpp
@@ -1474,26 +1474,27 @@ PyTypeObject InitializePyWrapperSetType()
 		}
 	};
 
+	// NOTE: _ElemType = typing.TypeVar('_ElemType') and Type/Iterable are defined by the Python typing module.
 	static PyMethodDef PyMethods[] = {
-		{ "cast", PyCFunctionCast(&FMethods::Cast), METH_VARARGS | METH_KEYWORDS | METH_CLASS, "X.cast(type, obj) -> Set -- cast the given object to this Unreal set type" },
-		{ "__copy__", PyCFunctionCast(&FMethods::Copy), METH_NOARGS, "x.__copy__() -> Set -- copy this Unreal set" },
-		{ "copy", PyCFunctionCast(&FMethods::Copy), METH_NOARGS, "x.copy() -> Set -- copy this Unreal set" },
-		{ "add", PyCFunctionCast(&FMethods::Add), METH_VARARGS, "x.add(value) -> None -- add the given value to this Unreal set if not already present" },
-		{ "discard", PyCFunctionCast(&FMethods::Discard), METH_VARARGS, "x.discard(value) -> None -- remove the given value from this Unreal set, or do nothing if not present" },
-		{ "remove", PyCFunctionCast(&FMethods::Remove), METH_VARARGS, "x.remove(value) -> None -- remove the given value from this Unreal set, or raise KeyError if not present" },
-		{ "pop", PyCFunctionCast(&FMethods::Pop), METH_NOARGS, "x.pop() -> value -- remove and return an arbitrary value from this Unreal set, or raise KeyError if the set is empty" },
-		{ "clear", PyCFunctionCast(&FMethods::Clear), METH_NOARGS, "x.clear() -> None -- remove all values from this Unreal set" },
-		{ "difference", PyCFunctionCast(&FMethods::Difference), METH_VARARGS, "x.difference(...) -> Set -- return the difference between this Unreal set and the other iterables passed for comparison (ie, all values that are in this set but not the others)" },
-		{ "difference_update", PyCFunctionCast(&FMethods::DifferenceUpdate), METH_VARARGS, "x.difference_update(...) -> None -- make this set the difference between this Unreal set and the other iterables passed for comparison (ie, all values that are in this set but not the others)" },
-		{ "intersection", PyCFunctionCast(&FMethods::Intersection), METH_VARARGS, "x.intersection(...) -> Set -- return the intersection between this Unreal set and the other iterables passed for comparison (ie, values that are common to all of the sets)" },
-		{ "intersection_update", PyCFunctionCast(&FMethods::IntersectionUpdate), METH_VARARGS, "x.intersection_update(...) -> None -- make this set the intersection between this Unreal set and the other iterables passed for comparison (ie, values that are common to all of the sets)" },
-		{ "symmetric_difference", PyCFunctionCast(&FMethods::SymmetricDifference), METH_VARARGS, "x.symmetric_difference(other) -> Set -- return the symmetric difference between this Unreal set and another (ie, values that are in exactly one of the sets)" },
-		{ "symmetric_difference_update", PyCFunctionCast(&FMethods::SymmetricDifferenceUpdate), METH_VARARGS, "x.symmetric_difference_update(other) -> None -- make this set the symmetric difference between this Unreal set and another (ie, values that are in exactly one of the sets)" },
-		{ "union", PyCFunctionCast(&FMethods::Union), METH_VARARGS, "x.union(...) -> Set -- return the union between this Unreal set and the other iterables passed for comparison (ie, values that are in any set)" },
-		{ "update", PyCFunctionCast(&FMethods::Update), METH_VARARGS, "x.update(...) -> None -- make this set the union between this Unreal set and the other iterables passed for comparison (ie, values that are in any set)" },
-		{ "isdisjoint", PyCFunctionCast(&FMethods::IsDisjoint), METH_VARARGS, "x.isdisjoint(other) -> bool -- return True if there is a null intersection between this Unreal set and another" },
-		{ "issubset", PyCFunctionCast(&FMethods::IsSubset), METH_VARARGS, "x.issubset(other) -> bool -- return True if another set contains this Unreal set" },
-		{ "issuperset", PyCFunctionCast(&FMethods::IsSuperset), METH_VARARGS, "x.issuperset(other) -> bool -- return True if this Unreal set contains another" },
+		{ "cast", PyCFunctionCast(&FMethods::Cast), METH_VARARGS | METH_KEYWORDS | METH_CLASS, "cast(cls, type: Type[_ElemType], obj: object) -> Set[_ElemType] -- cast the given object to this Unreal set type" },
+		{ "__copy__", PyCFunctionCast(&FMethods::Copy), METH_NOARGS, "__copy__(self) -> Set[_ElemType] -- copy this Unreal set" },
+		{ "copy", PyCFunctionCast(&FMethods::Copy), METH_NOARGS, "copy(self) -> Set[_ElemType] -- copy this Unreal set" },
+		{ "add", PyCFunctionCast(&FMethods::Add), METH_VARARGS, "add(self, value: _ElemType) -> None -- add the given value to this Unreal set if not already present" },
+		{ "discard", PyCFunctionCast(&FMethods::Discard), METH_VARARGS, "discard(self, value: _ElemType) -> None -- remove the given value from this Unreal set, or do nothing if not present" },
+		{ "remove", PyCFunctionCast(&FMethods::Remove), METH_VARARGS, "remove(self, value: _ElemType) -> None -- remove the given value from this Unreal set, or raise KeyError if not present" },
+		{ "pop", PyCFunctionCast(&FMethods::Pop), METH_NOARGS, "pop(self) -> _ElemType -- remove and return an arbitrary value from this Unreal set, or raise KeyError if the set is empty" },
+		{ "clear", PyCFunctionCast(&FMethods::Clear), METH_NOARGS, "clear(self) -> None -- remove all values from this Unreal set" },
+		{ "difference", PyCFunctionCast(&FMethods::Difference), METH_VARARGS, "difference(self, *iterable: Iterable[_ElemType]) -> Set[_ElemType] -- return the difference between this Unreal set and the other iterables passed for comparison (ie, all values that are in this set but not the others)" },
+		{ "difference_update", PyCFunctionCast(&FMethods::DifferenceUpdate), METH_VARARGS, "difference_update(self, *iterables: Iterable[_ElemType]) -> None -- make this set the difference between this Unreal set and the other iterables passed for comparison (ie, all values that are in this set but not the others)" },
+		{ "intersection", PyCFunctionCast(&FMethods::Intersection), METH_VARARGS, "intersection(self, *iterables: Iterable[_ElemType]) -> Set[_ElemType] -- return the intersection between this Unreal set and the other iterables passed for comparison (ie, values that are common to all of the sets)" },
+		{ "intersection_update", PyCFunctionCast(&FMethods::IntersectionUpdate), METH_VARARGS, "intersection_update(self, *iterables: Iterable[_ElemType]) -> None -- make this set the intersection between this Unreal set and the other iterables passed for comparison (ie, values that are common to all of the sets)" },
+		{ "symmetric_difference", PyCFunctionCast(&FMethods::SymmetricDifference), METH_VARARGS, "symmetric_difference(self, other: Iterable[_ElemType]) -> Set[_ElemType] -- return the symmetric difference between this Unreal set and another (ie, values that are in exactly one of the sets)" },
+		{ "symmetric_difference_update", PyCFunctionCast(&FMethods::SymmetricDifferenceUpdate), METH_VARARGS, "symmetric_difference_update(self, other: Iterable[_ElemType]) -> None -- make this set the symmetric difference between this Unreal set and another (ie, values that are in exactly one of the sets)" },
+		{ "union", PyCFunctionCast(&FMethods::Union), METH_VARARGS, "union(self, *iterables: Iterable[_ElemType]) -> Set[_ElemType] -- return the union between this Unreal set and the other iterables passed for comparison (ie, values that are in any set)" },
+		{ "update", PyCFunctionCast(&FMethods::Update), METH_VARARGS, "update(self, *iterables: Iterable[_ElemType]) -> None -- make this set the union between this Unreal set and the other iterables passed for comparison (ie, values that are in any set)" },
+		{ "isdisjoint", PyCFunctionCast(&FMethods::IsDisjoint), METH_VARARGS, "isdisjoint(self, other: Iterable[_ElemType]) -> bool -- return True if there is a null intersection between this Unreal set and another" },
+		{ "issubset", PyCFunctionCast(&FMethods::IsSubset), METH_VARARGS, "issubset(self, other: Iterable[_ElemType]) -> bool -- return True if another set contains this Unreal set" },
+		{ "issuperset", PyCFunctionCast(&FMethods::IsSuperset), METH_VARARGS, "issuperset(self, other: Iterable[_ElemType]) -> bool -- return True if this Unreal set contains another" },
 		{ nullptr, nullptr, 0, nullptr }
 	};
 
diff --git a/Engine/Plugins/Experimental/PythonScriptPlugin/Source/PythonScriptPlugin/Private/PyWrapperStruct.cpp b/Engine/Plugins/Experimental/PythonScriptPlugin/Source/PythonScriptPlugin/Private/PyWrapperStruct.cpp
--- a/Engine/Plugins/Experimental/PythonScriptPlugin/Source/PythonScriptPlugin/Private/PyWrapperStruct.cpp
+++ b/Engine/Plugins/Experimental/PythonScriptPlugin/Source/PythonScriptPlugin/Private/PyWrapperStruct.cpp
@@ -1260,19 +1260,20 @@ PyTypeObject InitializePyWrapperStructType()
 		}
 	};
 
+	// NOTE: _T = typing.TypeVar('_T') and Type/Any/Union/Mapping/Sequence/Tuple are defined by the Python typing module.
 	static PyMethodDef PyMethods[] = {
-		{ PyGenUtil::PostInitFuncName, PyCFunctionCast(&FMethods::PostInit), METH_NOARGS, "x._post_init() -> None -- called during Unreal struct initialization (equivalent to PostInitProperties in C++)" },
-		{ "cast", PyCFunctionCast(&FMethods::Cast), METH_VARARGS | METH_CLASS, "X.cast(object) -> struct -- cast the given object to this Unreal struct type" },
-		{ "static_struct", PyCFunctionCast(&FMethods::StaticStruct), METH_NOARGS | METH_CLASS, "X.static_struct() -> Struct -- get the Unreal struct of this type" },
-		{ "__copy__", PyCFunctionCast(&FMethods::Copy), METH_NOARGS, "x.__copy__() -> struct -- copy this Unreal struct" },
-		{ "copy", PyCFunctionCast(&FMethods::Copy), METH_NOARGS, "x.copy() -> struct -- copy this Unreal struct" },
-		{ "assign", PyCFunctionCast(&FMethods::Assign), METH_VARARGS, "x.assign(object) -> None -- assign the value of this Unreal struct to value of the given object" },
-		{ "to_tuple", PyCFunctionCast(&FMethods::ToTuple), METH_NOARGS, "x.to_tuple() -> tuple -- break this Unreal struct into a tuple of its properties" },
-		{ "get_editor_property", PyCFunctionCast(&FMethods::GetEditorProperty), METH_VARARGS | METH_KEYWORDS, "x.get_editor_property(name) -> object -- get the value of any property visible to the editor" },
-		{ "set_editor_property", PyCFunctionCast(&FMethods::SetEditorProperty), METH_VARARGS | METH_KEYWORDS, "x.set_editor_property(name, value, notify_mode=PropertyAccessChangeNotifyMode.DEFAULT) -> None -- set the value of any property visible to the editor, ensuring that the pre/post change notifications are called" },
-		{ "set_editor_properties", PyCFunctionCast(&FMethods::SetEditorProperties), METH_VARARGS, "x.set_editor_properties(property_info) -> None -- set the value of any properties visible to the editor (from a name->value dict), ensuring that the pre/post change notifications are called" },
-		{ "export_text", PyCFunctionCast(&FMethods::ExportToText), METH_NOARGS, "X.export_text() -> str -- exports the content of the Unreal struct of this type" },
-		{ "import_text", PyCFunctionCast(&FMethods::ImportFromText), METH_VARARGS | METH_KEYWORDS, "X.import_text(content) -> bool -- imports the provided string into the Unreal struct" },
+		{ PyGenUtil::PostInitFuncName, PyCFunctionCast(&FMethods::PostInit), METH_NOARGS, "_post_init(self) -> None -- called during Unreal struct initialization (equivalent to PostInitProperties in C++)" },
+		{ "cast", PyCFunctionCast(&FMethods::Cast), METH_VARARGS | METH_CLASS, "cast(cls: Type[_T], object: Union[object, Mapping[str, object], Iterable[object]]) -> _T -- cast the given object to this Unreal struct type. Can be partial Mapping[fieldName, fiedValue] or a sequence of field values" },
+		{ "static_struct", PyCFunctionCast(&FMethods::StaticStruct), METH_NOARGS | METH_CLASS, "static_struct(cls) -> ScriptStruct -- get the Unreal struct of this type" },
+		{ "__copy__", PyCFunctionCast(&FMethods::Copy), METH_NOARGS, "__copy__(self) -> Any -- copy this Unreal struct" },
+		{ "copy", PyCFunctionCast(&FMethods::Copy), METH_NOARGS, "copy(self) -> Any -- copy this Unreal struct" },
+		{ "assign", PyCFunctionCast(&FMethods::Assign), METH_VARARGS, "assign(self, other: object) -> None -- assign the value of this Unreal struct to value of the given object" },
+		{ "to_tuple", PyCFunctionCast(&FMethods::ToTuple), METH_NOARGS, "to_tuple(self) -> Tuple[object, ...] -- break this Unreal struct into a tuple of its properties" },
+		{ "get_editor_property", PyCFunctionCast(&FMethods::GetEditorProperty), METH_VARARGS | METH_KEYWORDS, "get_editor_property(self, name: Union[Name, str]) -> object -- get the value of any property visible to the editor" },
+		{ "set_editor_property", PyCFunctionCast(&FMethods::SetEditorProperty), METH_VARARGS | METH_KEYWORDS, "set_editor_property(self, name: Union[Name, str], value: object, notify_mode: PropertyAccessChangeNotifyMode=PropertyAccessChangeNotifyMode.DEFAULT) -> None -- set the value of any property visible to the editor, ensuring that the pre/post change notifications are called" },
+		{ "set_editor_properties", PyCFunctionCast(&FMethods::SetEditorProperties), METH_VARARGS, "set_editor_properties(self, properties: Mapping[str, object]) -> None -- set the value of any properties visible to the editor (from a name->value dict), ensuring that the pre/post change notifications are called" },
+		{ "export_text", PyCFunctionCast(&FMethods::ExportToText), METH_NOARGS, "export_text(self) -> str -- exports the content of the Unreal struct of this type" },
+		{ "import_text", PyCFunctionCast(&FMethods::ImportFromText), METH_VARARGS | METH_KEYWORDS, "import_text(self, content: str) -> bool -- imports the provided string into the Unreal struct" },
 		{ nullptr, nullptr, 0, nullptr }
 	};
 
diff --git a/Engine/Plugins/Experimental/PythonScriptPlugin/Source/PythonScriptPlugin/Private/PyWrapperText.cpp b/Engine/Plugins/Experimental/PythonScriptPlugin/Source/PythonScriptPlugin/Private/PyWrapperText.cpp
--- a/Engine/Plugins/Experimental/PythonScriptPlugin/Source/PythonScriptPlugin/Private/PyWrapperText.cpp
+++ b/Engine/Plugins/Experimental/PythonScriptPlugin/Source/PythonScriptPlugin/Private/PyWrapperText.cpp
@@ -440,20 +440,21 @@ PyTypeObject InitializePyWrapperTextType()
 		}
 	};
 
+	// NOTE: Union/Mapping are defines in the Python typing module.
 	static PyMethodDef PyMethods[] = {
-		{ "cast", PyCFunctionCast(&FMethods::Cast), METH_VARARGS | METH_CLASS, "X.cast(object) -> Text -- cast the given object to this Unreal text type" },
-		{ "as_number", PyCFunctionCast(&FMethods::AsNumber), METH_VARARGS | METH_CLASS, "X.as_number(num) -> Text -- convert the given number to a culture correct Unreal text representation" },
-		{ "as_percent", PyCFunctionCast(&FMethods::AsPercent), METH_VARARGS | METH_CLASS, "X.as_percent(num) -> Text -- convert the given number to a culture correct Unreal text percentgage representation" },
-		{ "as_currency", PyCFunctionCast(&FMethods::AsCurrency), METH_VARARGS | METH_CLASS, "X.as_currency(val, code) -> Text -- convert the given number (specified in the smallest unit for the given currency) to a culture correct Unreal text currency representation" },
+		{ "cast", PyCFunctionCast(&FMethods::Cast), METH_VARARGS | METH_CLASS, "cast(cls, object: object) -> Text -- cast the given object to this Unreal text type" },
+		{ "as_number", PyCFunctionCast(&FMethods::AsNumber), METH_VARARGS | METH_CLASS, "as_number(cls, num: float) -> Text -- convert the given number to a culture correct Unreal text representation" },
+		{ "as_percent", PyCFunctionCast(&FMethods::AsPercent), METH_VARARGS | METH_CLASS, "as_percent(cls, num: float) -> Text -- convert the given number to a culture correct Unreal text percentgage representation" },
+		{ "as_currency", PyCFunctionCast(&FMethods::AsCurrency), METH_VARARGS | METH_CLASS, "as_currency(cls, val: int, code: str) -> Text -- convert the given number (specified in the smallest unit for the given currency i.e. 650 for $6.50) to a culture correct Unreal text currency representation" },
 		// todo: datetime?
-		{ "is_empty", PyCFunctionCast(&FMethods::IsEmpty), METH_NOARGS, "x.is_empty() -> bool -- is this Unreal text empty?" },
-		{ "is_empty_or_whitespace", PyCFunctionCast(&FMethods::IsEmptyOrWhitespace), METH_NOARGS, "x.is_empty_or_whitespace() -> bool -- is this Unreal text empty or only whitespace?" },
-		{ "is_transient", PyCFunctionCast(&FMethods::IsTransient), METH_NOARGS, "x.is_transient() -> bool -- is this Unreal text transient?" },
-		{ "is_culture_invariant", PyCFunctionCast(&FMethods::IsEmptyOrWhitespace), METH_NOARGS, "x.is_culture_invariant() -> bool -- is this Unreal text culture invariant?" },
-		{ "is_from_string_table", PyCFunctionCast(&FMethods::IsFromStringTable), METH_NOARGS, "x.is_from_string_table() -> bool -- is this Unreal text referencing a string table entry?" },
-		{ "to_lower", PyCFunctionCast(&FMethods::ToLower), METH_NOARGS, "x.to_lower() -> Text -- convert this Unreal text to lowercase in a culture correct way" },
-		{ "to_upper", PyCFunctionCast(&FMethods::ToUpper), METH_NOARGS, "x.to_upper() -> Text -- convert this Unreal text to uppercase in a culture correct way" },
-		{ "format", PyCFunctionCast(&FMethods::Format), METH_VARARGS | METH_KEYWORDS, "x.format(...) -> Text -- use this Unreal text as a format pattern and generate a new text using the format arguments (may be a mapping, sequence, or set of (optionally named) arguments)" },
+		{ "is_empty", PyCFunctionCast(&FMethods::IsEmpty), METH_NOARGS, "is_empty(self) -> bool -- is this Unreal text empty?" },
+		{ "is_empty_or_whitespace", PyCFunctionCast(&FMethods::IsEmptyOrWhitespace), METH_NOARGS, "is_empty_or_whitespace(self) -> bool -- is this Unreal text empty or only whitespace?" },
+		{ "is_transient", PyCFunctionCast(&FMethods::IsTransient), METH_NOARGS, "is_transient(self) -> bool -- is this Unreal text transient?" },
+		{ "is_culture_invariant", PyCFunctionCast(&FMethods::IsEmptyOrWhitespace), METH_NOARGS, "is_culture_invariant(self) -> bool -- is this Unreal text culture invariant?" },
+		{ "is_from_string_table", PyCFunctionCast(&FMethods::IsFromStringTable), METH_NOARGS, "is_from_string_table(self) -> bool -- is this Unreal text referencing a string table entry?" },
+		{ "to_lower", PyCFunctionCast(&FMethods::ToLower), METH_NOARGS, "to_lower(self) -> Text -- convert this Unreal text to lowercase in a culture correct way" },
+		{ "to_upper", PyCFunctionCast(&FMethods::ToUpper), METH_NOARGS, "to_upper(self) -> Text -- convert this Unreal text to uppercase in a culture correct way" },
+		{ "format", PyCFunctionCast(&FMethods::Format), METH_VARARGS | METH_KEYWORDS, "format(self, *args: Union[Mapping[str, object], object], **named_args: object) -> Text -- use this Unreal text as a format pattern and generate a new text using the format arguments (may be a mapping[arg_name, arg] or comma separed (optionally named) arguments)" },
 		{ nullptr, nullptr, 0, nullptr }
 	};
 
diff --git a/Engine/Plugins/Experimental/PythonScriptPlugin/Source/PythonScriptPlugin/Private/PyWrapperTypeRegistry.cpp b/Engine/Plugins/Experimental/PythonScriptPlugin/Source/PythonScriptPlugin/Private/PyWrapperTypeRegistry.cpp
--- a/Engine/Plugins/Experimental/PythonScriptPlugin/Source/PythonScriptPlugin/Private/PyWrapperTypeRegistry.cpp
+++ b/Engine/Plugins/Experimental/PythonScriptPlugin/Source/PythonScriptPlugin/Private/PyWrapperTypeRegistry.cpp
@@ -16,6 +16,8 @@
 #include "PyConversion.h"
 #include "PyGIL.h"
 #include "PyCore.h"
+#include "PyEditor.h"
+#include "PyEngine.h"
 #include "PyFileWriter.h"
 #include "PythonScriptPluginSettings.h"
 #include "ProfilingDebugging/ScopedTimers.h"
@@ -38,9 +40,7 @@
 
 #if WITH_PYTHON
 
-namespace UE
-{
-namespace Python
+namespace UE::Python
 {
 
 /**
@@ -117,8 +117,287 @@ void SetBreakFunction(FPyWrapperStructMetaData& MetaData, UFunction* BreakFunc)
 };
 
 
-} // namespace Python
-} // namespace UE
+static bool GGeneratingOnlineDoc = false;
+
+void SetGeneratingOnlineDoc(bool bGeneratingOnlineDoc)
+{
+	GGeneratingOnlineDoc = bGeneratingOnlineDoc;
+}
+
+bool IsGeneratingOnlineDoc()
+{
+	return GGeneratingOnlineDoc;
+}
+
+bool ShouldEllipseParamDefaultValue(const FString& DefaultValue)
+{
+	// Some default values might not parse because they use undeclared types.
+	const bool bEnsureDefaultValueParses = (PyGenUtil::GetTypeHintingMode() == ETypeHintingMode::TypeChecker || IsGeneratingOnlineDoc());
+
+	if (bEnsureDefaultValueParses)
+	{
+		// Don't ellipse the default value if it is a simple Python expression. We want to the default value with "..." when the
+		// expression might refer to types that aren't declared yet. We can use undeclared type for type hint, but not for default
+		// value. The stub types are not written in a strict order so one type can refer another that was written later. For sake
+		// of simplicity, just keep simple default values, ellipse all others to the user detriment. :(
+		return	!(DefaultValue.IsNumeric() ||      // int or float
+					DefaultValue == TEXT("True") ||  // bool
+					DefaultValue == TEXT("False") || // bool
+					DefaultValue == TEXT("[]") ||    // list
+					DefaultValue == TEXT("{}") ||    // dict
+					DefaultValue == TEXT("()") ||    // tuple.
+					(DefaultValue.StartsWith("'") && DefaultValue.EndsWith("'")) || // string
+					(DefaultValue.StartsWith("\"") && DefaultValue.EndsWith("\""))); // string
+	}
+	return false;
+};
+
+/**
+ * Parses a Python doc string to find a method declaration that starts with the method name and ends with "--". For example
+ * if a given doc string is "xyz.log(arg: Any) -> None -- log the given argument as information in the LogPython category"
+ * the method "log(arg: Any) -> None" will be returned as the declaration if the searched method was "log".
+ */
+bool ParseMethodDeclarationFromDocString(const FString& InDocString, const FString& InMethodName, FString& OutMethodDecl)
+{
+	int32 BeginFuncDeclPos = InDocString.Find(InMethodName, ESearchCase::CaseSensitive, ESearchDir::FromStart, 0);
+	if (BeginFuncDeclPos == INDEX_NONE)
+	{
+		return false;
+	}
+
+	int32 EndFuncDeclPos = InDocString.Find(TEXT("--"), ESearchCase::CaseSensitive, ESearchDir::FromStart, BeginFuncDeclPos);
+	if (EndFuncDeclPos == INDEX_NONE)
+	{
+		return false;
+	}
+
+	OutMethodDecl = InDocString.Mid(BeginFuncDeclPos, EndFuncDeclPos - BeginFuncDeclPos);
+	return true;
+}
+
+/**
+ * Reads the human written doc string of a non-reflected function and try to extract a well-formed methods declaration containing
+ * type hinting. The function aims to generate type hints for the raw C/C++ functions that are not reflected but still exposed. If
+ * the function fails to find a well-formed and type hinted declaration in the doc string, it returns the generic signature
+ * 'def funcName(*args, **kargs):' that doesn't hint for types, but will still allow a form of auto-completion in Python text editors.
+ * Since raw C/C++ function doc strings are handwritten, they need extra caution to avoid errors. The function cannot check the
+ * types being handwritten, so the burden is put on the function provider to keep the declaration accurate.
+ * 
+ * To be successfully parsed, the type hinted declaration from the doc string must be well formed. See format below in the implementation.
+ * Any malformed declaration will be rejected and the function will fallback on the genric form.
+ */
+FString GenerateMethodDecl(const PyMethodDef* MethodDef, const FString& ModuleOrClass, bool bIsClassMethod = false, bool bIsMemberMethod = false)
+{
+	// Parse the Python method declared in the doc string check if it is a valid type hinted declaration and possibly fix it to make it well-formed.
+	auto TryParsingDocStringMethodDecl = [bIsClassMethod, bIsMemberMethod](const FStringView& InMethodDecl, const FStringView& InMethodName, FString& OutMethodDeclWithHints)
+	{
+		int32 OpenParenthesisPos = InMethodDecl.Find(TEXT("("));
+		if (OpenParenthesisPos == INDEX_NONE)
+		{
+			return false;
+		}
+
+		int32 CloseParenthesisPos = INDEX_NONE;
+		if (!InMethodDecl.FindLastChar(TEXT(')'), CloseParenthesisPos))
+		{
+			return false;
+		}
+
+		FStringView ReturnTypeDelimiter(TEXT("->"));
+		int32 ReturnTypeDelimiterPos = InMethodDecl.Find(ReturnTypeDelimiter, CloseParenthesisPos);
+		if (ReturnTypeDelimiterPos == INDEX_NONE)
+		{
+			return false;
+		}
+
+		// Check if the return type is hinted.
+		int32 BeginReturnTypePos = ReturnTypeDelimiterPos + ReturnTypeDelimiter.Len(); // After the ->
+		FStringView ReturnType = InMethodDecl.SubStr(BeginReturnTypePos, InMethodDecl.Len() - BeginReturnTypePos).TrimStartAndEnd();
+		if (ReturnType.Len() <= 0)
+		{
+			return false;
+		}
+
+		int32 ParamCount = 0;
+		FString Params;
+
+		bool bClsOrSelfParamFound = false;
+
+		// Check if the parameters are correcly hinted.
+		int32 ParamsLen = CloseParenthesisPos - OpenParenthesisPos;
+		if (ParamsLen > 1)
+		{
+			FStringView AllParamsString = InMethodDecl.SubStr(OpenParenthesisPos + 1, ParamsLen - 1); // +1 and -1 to exclude the parenthesis.
+
+			bool bLastParamReached = false;
+			int32 StartParamPos = 0;
+			do
+			{
+				// Find the end of a parameter declaration, caring for commas that aren't parameter separator, but part of the type hint, like in
+				// Callable[[int, int], str] or Union[int, str, float]. We need to skip the content of [] to find the end of a parameter declaration.
+				int BracketCount = 0;
+				int CurrPos = StartParamPos;
+				int EndParamPos = INDEX_NONE;
+				for (TCHAR Ch : AllParamsString.SubStr(StartParamPos, AllParamsString.Len() - StartParamPos))
+				{
+					if (Ch == TEXT('['))
+					{
+						++BracketCount;
+					}
+					else if (Ch == TEXT(']'))
+					{
+						--BracketCount;
+					}
+					else if (Ch == TEXT(','))
+					{
+						if (BracketCount == 0) // Outside [], the comma separates parameters.
+						{
+							EndParamPos = CurrPos;
+							break;
+						}
+					}
+					++CurrPos;
+				}
+
+				// Check if the parameters declaration is correctly formed.
+				FStringView ParamDecl;
+
+				if (EndParamPos != INDEX_NONE)
+				{
+					ParamDecl = AllParamsString.SubStr(StartParamPos, EndParamPos - StartParamPos);
+				}
+				else
+				{
+					ParamDecl = AllParamsString.SubStr(StartParamPos, AllParamsString.Len() - StartParamPos);
+					bLastParamReached = true;
+				}
+
+				// Search for parameter type hint separator.
+				const int32 NameTypeSeparatorPos = ParamDecl.Find(TEXT(":"));
+
+				// Those special doesn't require type hint, but they can if needed.
+				bool bIsSelf = ParamDecl.StartsWith(TEXT("self"));
+				bool bIsCls = ParamDecl.StartsWith(TEXT("cls"));
+				bool bParamTypeHinted = NameTypeSeparatorPos != INDEX_NONE;
+				bClsOrSelfParamFound |= bIsSelf || bIsCls;
+
+				if (!bParamTypeHinted && !bIsSelf && !bIsCls)
+				{
+					return false; // Invalid declaration: the param is not type hinted.
+				}
+
+				FStringView ParamName;
+				FStringView ParamTypeAndOptDefaultValue;
+				if (bParamTypeHinted)
+				{
+					ParamName = ParamDecl.SubStr(0, NameTypeSeparatorPos).TrimStartAndEnd();
+					ParamTypeAndOptDefaultValue = ParamDecl.SubStr(NameTypeSeparatorPos + 1, ParamDecl.Len() - NameTypeSeparatorPos).TrimStartAndEnd();
+
+					if (ParamName.Len() == 0 || ParamTypeAndOptDefaultValue.Len() == 0)
+					{
+						return false;
+					}
+				}
+				else // Only 'self' or 'cls' case should remain.
+				{
+					ParamName = ParamDecl.TrimStartAndEnd();
+					if (ParamName != TEXT("self") && ParamName != TEXT("cls"))
+					{
+						return false;
+					}
+				}
+
+				// Check if the paramType declaration also embeds a default value.
+				const int32 DefaultValueSeparatorPos = ParamTypeAndOptDefaultValue.Find(TEXT("="));
+				if (DefaultValueSeparatorPos == INDEX_NONE)
+				{
+					Params += PyGenUtil::PythonizeMethodParam(FString(ParamName), FString(ParamTypeAndOptDefaultValue)); // No default value found.
+				}
+				else
+				{
+					FString ParamType(ParamTypeAndOptDefaultValue.SubStr(0, DefaultValueSeparatorPos).TrimStartAndEnd());
+					FString ParamDefaultValue(ParamTypeAndOptDefaultValue.SubStr(DefaultValueSeparatorPos + 1, ParamTypeAndOptDefaultValue.Len() - DefaultValueSeparatorPos).TrimStartAndEnd());
+					if (ShouldEllipseParamDefaultValue(ParamDefaultValue))
+					{
+						ParamDefaultValue = TEXT("...");
+					}
+					Params += PyGenUtil::PythonizeMethodParam(FString(ParamName), FString(ParamType), FString(ParamDefaultValue));
+				}
+
+				if (!bLastParamReached)
+				{
+					Params += TEXT(", ");
+				}
+
+				StartParamPos = EndParamPos + 1; // +1 to start after the comma.
+				++ParamCount;
+			}
+			while (!bLastParamReached);
+		}
+
+		OutMethodDeclWithHints = TEXT("def ");
+		OutMethodDeclWithHints += InMethodName;
+		OutMethodDeclWithHints += TEXT("(");
+
+		if (bIsClassMethod && !bClsOrSelfParamFound)
+		{
+			OutMethodDeclWithHints += FString::Printf(TEXT("cls%s"), ParamCount > 0 ? TEXT(", ") : TEXT(""));
+		}
+		else if (bIsMemberMethod && !bClsOrSelfParamFound)
+		{
+			OutMethodDeclWithHints += FString::Printf(TEXT("self%s"), ParamCount > 0 ? TEXT(", ") : TEXT(""));
+		}
+
+		// All parameters type hint format was properly formed. Add them to the declaration.
+		OutMethodDeclWithHints += Params;
+
+		OutMethodDeclWithHints += TEXT(")");
+		OutMethodDeclWithHints += TEXT(" -> ");
+		OutMethodDeclWithHints += ReturnType;
+		OutMethodDeclWithHints += TEXT(":");
+
+		// The DocString contains a function declaration with properly formed with type hints. It is assumed to be a 'correct' declaration, but
+		// that doc string is handwritten, so human errors are possible and types might be incorrect.
+		return true;
+	};
+
+	FString DocString = UTF8_TO_TCHAR(MethodDef->ml_doc);
+	FString MethodName = UTF8_TO_TCHAR(MethodDef->ml_name);
+	FString MethodDecl;
+
+	if (PyGenUtil::IsTypeHintingEnabled())
+	{
+		// Try extracting type hinting from the doc string. If the Doc string appears to be well-formed, it is likely because the writer
+		// took time to type-hint it correctly.
+		if (ParseMethodDeclarationFromDocString(DocString, MethodName, MethodDecl))
+		{
+			FString WellFormedMethodDeclWithHints;
+			if (TryParsingDocStringMethodDecl(MethodDecl, MethodName, WellFormedMethodDeclWithHints))
+			{
+				return WellFormedMethodDeclWithHints;
+			}
+		}
+	}
+
+	// Failed to extract a well-formed type hinted method declaration from the doc string. Fallback to the generic and safe declaration.
+	const bool bHasKeywords = !!(MethodDef->ml_flags & METH_KEYWORDS);
+	if (bIsClassMethod)
+	{
+		MethodDecl = FString::Printf(TEXT("def %s(cls, *args%s):"), *MethodName, bHasKeywords ? TEXT(", **kwargs") : TEXT(""));
+	}
+	else if (bIsMemberMethod)
+	{
+		MethodDecl = FString::Printf(TEXT("def %s(self, *args%s):"), *MethodName, bHasKeywords ? TEXT(", **kwargs") : TEXT(""));
+	}
+	else
+	{
+		MethodDecl = FString::Printf(TEXT("def %s(*args%s):"), *MethodName, bHasKeywords ? TEXT(", **kwargs") : TEXT(""));
+	}
+
+	return MethodDecl;
+}
+
+} // namespace UE::Python
 
 
 
@@ -1217,7 +1496,7 @@ PyTypeObject* FPyWrapperTypeRegistry::GenerateWrappedClassType(const UClass* InC
 		{
 			if (HostClass->IsChildOf(InFunc->GetOwnerClass()))
 			{
-				REPORT_PYTHON_GENERATION_ISSUE(Error, TEXT("Function '%s.%s' is marked as 'ScriptConstantHost' but the host type (%s) is a child of the the class type of the static function. This is not allowed."), *InFunc->GetOwnerClass()->GetName(), *InFunc->GetName(), *HostClass->GetName());
+				REPORT_PYTHON_GENERATION_ISSUE(Error, TEXT("Function '%s.%s' is marked as 'ScriptConstantHost' but the host type (%s) is a child of the class type of the static function. This is not allowed."), *InFunc->GetOwnerClass()->GetName(), *InFunc->GetName(), *HostClass->GetName());
 				return;
 			}
 		}
@@ -2307,6 +2586,12 @@ void FPyWrapperTypeRegistry::GenerateStubCodeForWrappedTypes(const EPyOnlineDocs
 	FPyScopedGIL GIL;
 	FPyFileWriter PythonScript;
 
+	if (PyGenUtil::IsTypeHintingEnabled())
+	{
+		// Must be at the top of the file. This allows type hinting using types not yet declared. Like Object, Text, Name, etc.
+		PythonScript.WriteLine(TEXT("from __future__ import annotations")); // https://docs.python.org/3.7/whatsnew/3.7.html#pep-563-postponed-evaluation-of-annotations)
+	}
+
 	TUniquePtr<FPyOnlineDocsWriter> OnlineDocsWriter;
 	TSharedPtr<FPyOnlineDocsModule> OnlineDocsUnrealModule;
 	TSharedPtr<FPyOnlineDocsSection> OnlineDocsNativeTypesSection;
@@ -2324,6 +2609,8 @@ void FPyWrapperTypeRegistry::GenerateStubCodeForWrappedTypes(const EPyOnlineDocs
 		OnlineDocsClassTypesSection = OnlineDocsWriter->CreateSection(TEXT("Class Types"));
 		OnlineDocsEnumTypesSection = OnlineDocsWriter->CreateSection(TEXT("Enum Types"));
 		OnlineDocsDelegateTypesSection = OnlineDocsWriter->CreateSection(TEXT("Delegate Types"));
+
+		UE::Python::SetGeneratingOnlineDoc(true);
 	}
 
 	// Process additional Python files
@@ -2340,7 +2627,7 @@ void FPyWrapperTypeRegistry::GenerateStubCodeForWrappedTypes(const EPyOnlineDocs
 		{
 			const FString PythonBaseModuleName = PyGenUtil::GetModulePythonName(ModuleName, false);
 			const FString PythonModuleName = FString::Printf(TEXT("unreal_%s"), *PythonBaseModuleName);
-		
+
 			FString ModuleFilename;
 			if (PyUtil::IsModuleAvailableForImport(*PythonModuleName, &PyUtil::GetOnDiskUnrealModulesCache(), &ModuleFilename))
 			{
@@ -2413,6 +2700,16 @@ void FPyWrapperTypeRegistry::GenerateStubCodeForWrappedTypes(const EPyOnlineDocs
 		{
 			PythonScript.WriteNewLine();
 		}
+
+		if (PyGenUtil::IsTypeHintingEnabled())
+		{
+			PythonScript.WriteLine(TEXT("from typing import Any, Callable, Dict, ItemsView, Iterable, Iterator, KeysView, List, Mapping, MutableMapping, MutableSequence, MutableSet, Optional, Set, Tuple, Type, TypeVar, Union, ValuesView"));
+			PythonScript.WriteLine(TEXT("_T = TypeVar('_T')"));                 // For casting operations.
+			PythonScript.WriteLine(TEXT("_ElemType = TypeVar('_ElemType')"));   // For Array, FixedArray, Set classes along with casting operations.
+			PythonScript.WriteLine(TEXT("_KeyType = TypeVar('_KeyType')"));     // For Mapping[KV, KT] used by Map class.
+			PythonScript.WriteLine(TEXT("_ValueType = TypeVar('_ValueType')")); // For Mapping[KV, KT] used by Map class.
+			PythonScript.WriteNewLine();
+		}
 	}
 
 	// Process native glue code
@@ -2424,11 +2721,12 @@ void FPyWrapperTypeRegistry::GenerateStubCodeForWrappedTypes(const EPyOnlineDocs
 	{
 		for (PyMethodDef* MethodDef = NativePythonModule.PyModuleMethods; MethodDef && MethodDef->ml_name; ++MethodDef)
 		{
-			const bool bHasKeywords = !!(MethodDef->ml_flags & METH_KEYWORDS);
-			PythonScript.WriteLine(FString::Printf(TEXT("def %s(*args%s):"), UTF8_TO_TCHAR(MethodDef->ml_name), bHasKeywords ? TEXT(", **kwargs") : TEXT("")));
+			// Try parsing the doc string to extract the function declaration with type hinting.
+			FString MethodDecl = UE::Python::GenerateMethodDecl(MethodDef, TEXT("unreal"));
+			PythonScript.WriteLine(MethodDecl);
 			PythonScript.IncreaseIndent();
 			PythonScript.WriteDocString(UTF8_TO_TCHAR(MethodDef->ml_doc));
-			PythonScript.WriteLine(TEXT("pass"));
+			PythonScript.WriteLine(TEXT("...")); // pass
 			PythonScript.DecreaseIndent();
 			PythonScript.WriteNewLine();
 
@@ -2553,7 +2851,59 @@ void FPyWrapperTypeRegistry::GenerateStubCodeForWrappedTypes(const EPyOnlineDocs
 void FPyWrapperTypeRegistry::GenerateStubCodeForWrappedType(PyTypeObject* PyType, const PyGenUtil::FGeneratedWrappedType* GeneratedTypeData, FPyFileWriter& OutPythonScript, FPyOnlineDocsSection* OutOnlineDocsSection)
 {
 	const FString PyTypeName = UTF8_TO_TCHAR(PyType->tp_name);
-	OutPythonScript.WriteLine(FString::Printf(TEXT("class %s(%s):"), *PyTypeName, UTF8_TO_TCHAR(PyType->tp_base->tp_name)));
+
+	// If type hinting is configured for IDE auto-completion:
+	//     - Be nice to the user, don't bother him with type coercion (typing.Union) or None type validation (typing.Optional), show the raw Unreal types.
+	//     - Be nice to the user, don't ellipse non-parsing default values.
+	// 
+	// If type hinting is configure for type checking:
+	//    - Support known type coercions for method parameters (typing.Union[])
+	//    - Tag Object-like object as potentially 'None' using typing.Optional[]. We cannot know when an Object-like an input parameter or return value can be None,
+	//      so opt for the safe side and mark them all as Optional. That's might mislead users into thinking that some parameter/values can be None while in reality, they can't.
+	//    - Don't coerce returned type. Our API which largely reflect C++ returns one defined type, not something like Union[Name, str]
+	//    - Returned unreal.Array/unreal.Set/unreal.Map are a bit annoying to hint. For example, Array[Name] is hinted to accepts Name
+	//      param for example 'array.append("a name")' will fail type checker even though this works. User will have to do 'a.append(Name("my name"))'
+	//    - If an API returns an unreal.Array[T]/unreal.Set[T]/unreal.Map[T] and T is a type that can None such as an unreal.Array[unreal.Object],
+	//      then the return type will be Array[Optional[Object]] to allow having None type in there, which is also a bit annoying to work with.
+	//
+	// If type hinting is off:
+	//    - Use strict typing for return values to aid auto-complete (we also only care about the type and
+	//      not the value, so structs can be default constructed)
+	//
+	// Always default construct date as default value to avoid setting the generated date as the default value.
+
+	const PyGenUtil::EPythonizeFlags PythonizeMethodParamFlags = PyGenUtil::IsTypeHintingEnabled()
+		? (PyGenUtil::GetTypeHintingMode() == ETypeHintingMode::AutoCompletion
+			? PyGenUtil::EPythonizeFlags::DefaultConstructDateTime | PyGenUtil::EPythonizeFlags::WithTypeHinting // For IDE auto-completion.
+			: PyGenUtil::EPythonizeFlags::DefaultConstructDateTime | PyGenUtil::EPythonizeFlags::WithTypeHinting | PyGenUtil::EPythonizeFlags::WithTypeCoercion | PyGenUtil::EPythonizeFlags::WithOptionalType) // For type checker.
+		: PyGenUtil::EPythonizeFlags::DefaultConstructDateTime; // hinting is off.
+
+	const PyGenUtil::EPythonizeFlags PythonizeMethodReturnTypeFlags = PyGenUtil::IsTypeHintingEnabled()
+		? (PyGenUtil::GetTypeHintingMode() == ETypeHintingMode::AutoCompletion
+			? PyGenUtil::EPythonizeFlags::WithTypeHinting  // For IDE auto-completion.
+			: PyGenUtil::EPythonizeFlags::WithTypeHinting | PyGenUtil::EPythonizeFlags::WithOptionalType) // For type checker.
+		: PyGenUtil::EPythonizeFlags::DefaultConstructStructs | PyGenUtil::EPythonizeFlags::DefaultConstructDateTime | PyGenUtil::EPythonizeFlags::UseStrictTyping; // hinting is off.
+
+	if (PyGenUtil::IsTypeHintingEnabled() &&  (PyTypeName == TEXT("Array") || PyTypeName == TEXT("FixedArray")))
+	{
+		// Implies that _ElemType = TypeVar('_ElemType') is already written to the file.
+		OutPythonScript.WriteLine(FString::Printf(TEXT("class %s(%s, MutableSequence[_ElemType]):"), *PyTypeName, UTF8_TO_TCHAR(PyType->tp_base->tp_name)));
+	}
+	else if (PyGenUtil::IsTypeHintingEnabled() && PyTypeName == TEXT("Set"))
+	{
+		// Implies that _ElemType = TypeVar('_ElemType') is already written to the file.
+		OutPythonScript.WriteLine(FString::Printf(TEXT("class %s(%s, MutableSet[_ElemType]):"), *PyTypeName, UTF8_TO_TCHAR(PyType->tp_base->tp_name)));
+	}
+	else if (PyGenUtil::IsTypeHintingEnabled() && PyTypeName == TEXT("Map"))
+	{
+		// Implies that: _KeyType = TypeVar('_KeyType') and _ValueType = TypeVar('_ValueType') are already written to the file.
+		OutPythonScript.WriteLine(FString::Printf(TEXT("class %s(%s, MutableMapping[_KeyType, _ValueType]):"), *PyTypeName, UTF8_TO_TCHAR(PyType->tp_base->tp_name)));
+	}
+	else
+	{
+		OutPythonScript.WriteLine(FString::Printf(TEXT("class %s(%s):"), *PyTypeName, UTF8_TO_TCHAR(PyType->tp_base->tp_name)));
+	}
+
 	OutPythonScript.IncreaseIndent();
 	OutPythonScript.WriteDocString(UTF8_TO_TCHAR(PyType->tp_doc));
 
@@ -2570,7 +2920,7 @@ void FPyWrapperTypeRegistry::GenerateStubCodeForWrappedType(PyTypeObject* PyType
 		}
 		
 		// We use strict typing for return values to aid auto-complete (we also only care about the type and not the value, so structs can be default constructed)
-		static const uint32 PythonizeValueFlags = PyGenUtil::EPythonizeValueFlags::UseStrictTyping | PyGenUtil::EPythonizeValueFlags::DefaultConstructStructs;
+		static const PyGenUtil::EPythonizeFlags PythonizeValueFlags = PyGenUtil::EPythonizeFlags::UseStrictTyping | PyGenUtil::EPythonizeFlags::DefaultConstructStructs;
 
 		// If we have multiple return values and the main return value is a bool, skip it (to mimic PyGenUtils::PackReturnValues)
 		int32 ReturnPropIndex = 0;
@@ -2586,7 +2936,7 @@ void FPyWrapperTypeRegistry::GenerateStubCodeForWrappedType(PyTypeObject* PyType
 			const PyGenUtil::FGeneratedWrappedMethodParameter& ReturnParam = InOutputParams[ReturnPropIndex];
 			return PyGenUtil::PythonizeValue(ReturnParam.ParamProp, ReturnParam.ParamProp->ContainerPtrToValuePtr<void>(InBaseParamsAddr), PythonizeValueFlags);
 		}
-		else
+		else // Returns a tuple
 		{
 			FString FunctionReturnStr = TEXT("(");
 			for (int32 PackedPropIndex = 0; ReturnPropIndex < InOutputParams.Num(); ++ReturnPropIndex, ++PackedPropIndex)
@@ -2603,71 +2953,81 @@ void FPyWrapperTypeRegistry::GenerateStubCodeForWrappedType(PyTypeObject* PyType
 		}
 	};
 
-	auto ExportConstantValue = [&OutPythonScript, &PyTypeName](const TCHAR* InConstantName, const TCHAR* InConstantDocString, const FString& InConstantValue)
+	auto ExportConstant = [&OutPythonScript, &PyTypeName](const TCHAR* InConstantName, const FString& InConstantType, const FString& InConstantValue, const TCHAR* InConstantDocString)
 	{
 		FString ConstantValue(InConstantValue);
 
 		// Ensure that constant type is not same type as host type
 		if (ConstantValue.StartsWith(PyTypeName, ESearchCase::CaseSensitive) && (ConstantValue[PyTypeName.Len()] == TEXT('(')))
 		{
-			ConstantValue = TEXT("None");
+			ConstantValue = PyGenUtil::IsTypeHintingEnabled() ? TEXT("") : TEXT("None");
 		}
 
 		if (*InConstantDocString == 0)
 		{
 			// No docstring
-			OutPythonScript.WriteLine(FString::Printf(TEXT("%s = %s"), InConstantName, *ConstantValue));
+			ConstantValue.IsEmpty()
+				? OutPythonScript.WriteLine(FString::Printf(TEXT("%s%s"), InConstantName, InConstantType.IsEmpty() ? TEXT("") : *FString::Printf(TEXT(": %s"), *InConstantType)))
+				: OutPythonScript.WriteLine(FString::Printf(TEXT("%s%s = %s"), InConstantName, InConstantType.IsEmpty() ? TEXT("") : *FString::Printf(TEXT(": %s"), *InConstantType), *ConstantValue));
+			
 		}
 		else
 		{
 			if (FCString::Strchr(InConstantDocString, TEXT('\n')))
 			{
 				// Multi-line docstring
-				OutPythonScript.WriteLine(FString::Printf(TEXT("%s = %s"), InConstantName, *ConstantValue));
+				ConstantValue.IsEmpty()
+					? OutPythonScript.WriteLine(FString::Printf(TEXT("%s%s"), InConstantName, InConstantType.IsEmpty() ? TEXT("") : *FString::Printf(TEXT(": %s"), *InConstantType)))
+					: OutPythonScript.WriteLine(FString::Printf(TEXT("%s%s = %s"), InConstantName, InConstantType.IsEmpty() ? TEXT("") : *FString::Printf(TEXT(": %s"), *InConstantType), *ConstantValue));
+
 				OutPythonScript.WriteDocString(InConstantDocString);
 				OutPythonScript.WriteNewLine();
 			}
 			else
 			{
 				// Single-line docstring
-				OutPythonScript.WriteLine(FString::Printf(TEXT("%s = %s  #: %s"), InConstantName, *ConstantValue, InConstantDocString));
+				ConstantValue.IsEmpty()
+					? OutPythonScript.WriteLine(FString::Printf(TEXT("%s%s #: %s"), InConstantName, InConstantType.IsEmpty() ? TEXT("") : *FString::Printf(TEXT(": %s"), *InConstantType), InConstantDocString))
+					: OutPythonScript.WriteLine(FString::Printf(TEXT("%s%s = %s #: %s"), InConstantName, InConstantType.IsEmpty() ? TEXT("") : *FString::Printf(TEXT(": %s"), *InConstantType), *ConstantValue, InConstantDocString));
 			}
 		}
 	};
 
-	auto ExportGetSet = [&OutPythonScript](const TCHAR* InGetSetName, const TCHAR* InGetSetDocString, const TCHAR* InGetReturnValue, const bool InReadOnly)
+	auto ExportGetSet = [&OutPythonScript](const TCHAR* InGetSetMethodName, const TCHAR* InGetSetDocString, const TCHAR* InGetterReturnValue, const TCHAR* InGetterReturnType, const TCHAR* InSetterParam, const bool InReadOnly)
 	{
 		// Getter
 		OutPythonScript.WriteLine(TEXT("@property"));
-		OutPythonScript.WriteLine(FString::Printf(TEXT("def %s(self):"), InGetSetName));
+		OutPythonScript.WriteLine(FString::Printf(TEXT("def %s(self)%s%s:"), InGetSetMethodName, *InGetterReturnType != 0 ? TEXT(" -> ") : TEXT(""), InGetterReturnType));
 		OutPythonScript.IncreaseIndent();
 		OutPythonScript.WriteDocString(InGetSetDocString);
-		OutPythonScript.WriteLine(FString::Printf(TEXT("return %s"), InGetReturnValue));
+		PyGenUtil::IsTypeHintingEnabled()
+			? OutPythonScript.WriteLine(TEXT("...")) // No need to add a return value to help auto-complete/type checker when the type is hinted in the signature.
+			: OutPythonScript.WriteLine(FString::Printf(TEXT("return %s"), InGetterReturnValue));
 		OutPythonScript.DecreaseIndent();
 
 		if (!InReadOnly)
 		{
 			// Setter
-			OutPythonScript.WriteLine(FString::Printf(TEXT("@%s.setter"), InGetSetName));
-			OutPythonScript.WriteLine(FString::Printf(TEXT("def %s(self, value):"), InGetSetName));
+			OutPythonScript.WriteLine(FString::Printf(TEXT("@%s.setter"), InGetSetMethodName));
+			PyGenUtil::IsTypeHintingEnabled()
+				? OutPythonScript.WriteLine(FString::Printf(TEXT("def %s(self, %s) -> None:"), InGetSetMethodName, *InSetterParam != 0 ? InSetterParam : TEXT("value")))
+				: OutPythonScript.WriteLine(FString::Printf(TEXT("def %s(self, value):"), InGetSetMethodName));
 			OutPythonScript.IncreaseIndent();
-			OutPythonScript.WriteLine(TEXT("pass"));
+			OutPythonScript.WriteLine(TEXT("...")); // pass
 			OutPythonScript.DecreaseIndent();
 		}
 	};
 
-	auto ExportGeneratedMethod = [&OutPythonScript, &GetFunctionReturnValue](const PyGenUtil::FGeneratedWrappedMethod& InTypeMethod)
+	auto ExportGeneratedMethod = [PythonizeMethodParamFlags, PythonizeMethodReturnTypeFlags, &OutPythonScript, &GetFunctionReturnValue](const PyGenUtil::FGeneratedWrappedMethod& InTypeMethod)
 	{
 		FString MethodArgsStr;
 		FString FirstDefaultedParamName;
 		for (const PyGenUtil::FGeneratedWrappedMethodParameter& MethodParam : InTypeMethod.MethodFunc.InputParams)
 		{
 			MethodArgsStr += TEXT(", ");
-			MethodArgsStr += UTF8_TO_TCHAR(MethodParam.ParamName.GetData());
 			if (MethodParam.ParamDefaultValue.IsSet())
 			{
-				MethodArgsStr += TEXT("=");
-				MethodArgsStr += PyGenUtil::PythonizeDefaultValue(MethodParam.ParamProp, MethodParam.ParamDefaultValue.GetValue());
+				MethodArgsStr += PyGenUtil::PythonizeMethodParam(MethodParam, PythonizeMethodParamFlags, UE::Python::ShouldEllipseParamDefaultValue);
 				if (FirstDefaultedParamName.IsEmpty())
 				{
 					FirstDefaultedParamName = MethodParam.ParamProp->GetName();
@@ -2678,7 +3038,7 @@ void FPyWrapperTypeRegistry::GenerateStubCodeForWrappedType(PyTypeObject* PyType
 				// Some UFUNCTION has their default values declared as 'meta' invisible to C++ compiler, but Python generated glue will use that. Sometime, the next parameter(s) are
 				// not defaulted. This creates invalid Python declarations. In stub, this is legal to use "..." to indicate a default param value without specifying the exact value.
 				// That's a workaround for faulty method declarations and the C++ code should be fixed.
-				MethodArgsStr += TEXT("=...");
+				MethodArgsStr += PyGenUtil::PythonizeMethodParam(MethodParam, PythonizeMethodParamFlags | PyGenUtil::EPythonizeFlags::AddMissingDefaultValueEllipseWorkaround);
 
 				// Left as 'Verbose' until we fix known cases. Once fixed, we should turn this into a 'Warning'.
 				REPORT_PYTHON_GENERATION_ISSUE(Verbose, TEXT("The function '%s' is missing default values for param '%s'. All params after the first defaulted param '%s' (likely from meta data) should be defaulted. The Python stub defaulted missing default values."), 
@@ -2686,39 +3046,56 @@ void FPyWrapperTypeRegistry::GenerateStubCodeForWrappedType(PyTypeObject* PyType
 					MethodParam.ParamProp != nullptr ? *MethodParam.ParamProp->GetName() : TEXT("Unknown"),
 					*FirstDefaultedParamName);
 			}
+			else
+			{
+				MethodArgsStr += PyGenUtil::PythonizeMethodParam(MethodParam, PythonizeMethodParamFlags, UE::Python::ShouldEllipseParamDefaultValue);
+			}
 		}
 
-		FString MethodReturnStr;
+		// Either the type or the value, depending if type hinting is on/off.
+		FString MethodReturnValueStr;
+		FString MethodReturnType;
 		if (InTypeMethod.MethodFunc.Func)
 		{
-			FStructOnScope FuncParams(InTypeMethod.MethodFunc.Func);
-			MethodReturnStr = GetFunctionReturnValue(FuncParams.GetStructMemory(), InTypeMethod.MethodFunc.OutputParams);
+			if (PyGenUtil::IsTypeHintingEnabled())
+			{
+				MethodReturnType = PyGenUtil::PythonizeMethodReturnType(InTypeMethod.MethodFunc.OutputParams, PythonizeMethodReturnTypeFlags);
+			}
+			else
+			{
+				FStructOnScope FuncParams(InTypeMethod.MethodFunc.Func);
+				MethodReturnValueStr = GetFunctionReturnValue(FuncParams.GetStructMemory(), InTypeMethod.MethodFunc.OutputParams);
+			}
 		}
 		else
 		{
-			MethodReturnStr = TEXT("None");
+			MethodReturnType = TEXT("None");
+			MethodReturnValueStr = TEXT("None");
 		}
 
 		const bool bIsClassMethod = !!(InTypeMethod.MethodFlags & METH_CLASS);
 		if (bIsClassMethod)
 		{
 			OutPythonScript.WriteLine(TEXT("@classmethod"));
-			OutPythonScript.WriteLine(FString::Printf(TEXT("def %s(cls%s):"), UTF8_TO_TCHAR(InTypeMethod.MethodName.GetData()), *MethodArgsStr));
+			OutPythonScript.WriteLine(FString::Printf(TEXT("def %s(cls%s)%s%s:"), UTF8_TO_TCHAR(InTypeMethod.MethodName.GetData()), *MethodArgsStr, !MethodReturnType.IsEmpty() ? TEXT(" -> ") : TEXT(""), *MethodReturnType));
 		}
 		else
 		{
-			OutPythonScript.WriteLine(FString::Printf(TEXT("def %s(self%s):"), UTF8_TO_TCHAR(InTypeMethod.MethodName.GetData()), *MethodArgsStr));
+			OutPythonScript.WriteLine(FString::Printf(TEXT("def %s(self%s)%s%s:"), UTF8_TO_TCHAR(InTypeMethod.MethodName.GetData()), *MethodArgsStr, !MethodReturnType.IsEmpty() ? TEXT(" -> ") : TEXT(""), *MethodReturnType));
 		}
 		OutPythonScript.IncreaseIndent();
 		OutPythonScript.WriteDocString(UTF8_TO_TCHAR(InTypeMethod.MethodDoc.GetData()));
-		OutPythonScript.WriteLine(FString::Printf(TEXT("return %s"), *MethodReturnStr));
+		PyGenUtil::IsTypeHintingEnabled()
+			? OutPythonScript.WriteLine(TEXT("...")) // Omit the return value, the return type is hinted in the method signature.
+			: OutPythonScript.WriteLine(FString::Printf(TEXT("return %s"), *MethodReturnValueStr));
 		OutPythonScript.DecreaseIndent();
 	};
 
-	auto ExportGeneratedConstant = [&ExportConstantValue, &GetFunctionReturnValue](const PyGenUtil::FGeneratedWrappedConstant& InTypeConstant)
+	auto ExportGeneratedConstant = [PythonizeMethodReturnTypeFlags, &ExportConstant, &GetFunctionReturnValue](const PyGenUtil::FGeneratedWrappedConstant& InTypeConstant)
 	{
 		// Resolve the constant value
 		FString ConstantValueStr;
+		FString ConstantType;
 		if (InTypeConstant.ConstantFunc.Func && InTypeConstant.ConstantFunc.InputParams.Num() == 0)
 		{
 			UClass* Class = InTypeConstant.ConstantFunc.Func->GetOwnerClass();
@@ -2728,25 +3105,44 @@ void FPyWrapperTypeRegistry::GenerateStubCodeForWrappedType(PyTypeObject* PyType
 			PyUtil::InvokeFunctionCall(Obj, InTypeConstant.ConstantFunc.Func, FuncParams.GetStructMemory(), TEXT("export generated constant"));
 			PyErr_Clear(); // Clear any errors in case InvokeFunctionCall failed
 
-			ConstantValueStr = GetFunctionReturnValue(FuncParams.GetStructMemory(), InTypeConstant.ConstantFunc.OutputParams);
+			if (PyGenUtil::IsTypeHintingEnabled())
+			{
+				// If type hinting is enabled, just export the constant name and type.
+				ConstantType = PyGenUtil::PythonizeMethodReturnType(InTypeConstant.ConstantFunc.OutputParams, PythonizeMethodReturnTypeFlags);
+			}
+			else
+			{
+				ConstantValueStr = GetFunctionReturnValue(FuncParams.GetStructMemory(), InTypeConstant.ConstantFunc.OutputParams);
+			}
 		}
-		else
+		else // Cannot resove type/value.
 		{
+			ConstantType = TEXT("object");
 			ConstantValueStr = TEXT("None");
 		}
-		ExportConstantValue(UTF8_TO_TCHAR(InTypeConstant.ConstantName.GetData()), UTF8_TO_TCHAR(InTypeConstant.ConstantDoc.GetData()), ConstantValueStr);
+		ExportConstant(UTF8_TO_TCHAR(InTypeConstant.ConstantName.GetData()), ConstantType, ConstantValueStr, UTF8_TO_TCHAR(InTypeConstant.ConstantDoc.GetData()));
 	};
 
-	auto ExportGeneratedGetSet = [&ExportGetSet](const PyGenUtil::FGeneratedWrappedGetSet& InGetSet)
+	auto ExportGeneratedGetSet = [PythonizeMethodParamFlags, PythonizeMethodReturnTypeFlags, &ExportGetSet](const PyGenUtil::FGeneratedWrappedGetSet& InGetSet)
 	{
-		// We use strict typing for return values to aid auto-complete (we also only care about the type and not the value, so structs can be default constructed)
-		static const uint32 PythonizeValueFlags = PyGenUtil::EPythonizeValueFlags::UseStrictTyping | PyGenUtil::EPythonizeValueFlags::DefaultConstructStructs;
-		const FString GetReturnValue = PyGenUtil::PythonizeDefaultValue(InGetSet.Prop.Prop, FString(), PythonizeValueFlags);
 		const bool bIsReadOnly = InGetSet.Prop.Prop->HasAnyPropertyFlags(PropertyAccessUtil::RuntimeReadOnlyFlags);
-		ExportGetSet(UTF8_TO_TCHAR(InGetSet.GetSetName.GetData()), UTF8_TO_TCHAR(InGetSet.GetSetDoc.GetData()), *GetReturnValue, bIsReadOnly);
+		const FString SetterParamDecl = PyGenUtil::PythonizeMethodParam(TEXT("value"), InGetSet.Prop.Prop, PythonizeMethodParamFlags);
+		if (PyGenUtil::IsTypeHintingEnabled())
+		{
+			// No need to add a return value to help auto-complete/type checker when the type is hinted in the signature.
+			const FString GetterReturnType = PyGenUtil::PythonizeMethodReturnType(InGetSet.Prop.Prop, PythonizeMethodReturnTypeFlags);
+			ExportGetSet(UTF8_TO_TCHAR(InGetSet.GetSetName.GetData()), UTF8_TO_TCHAR(InGetSet.GetSetDoc.GetData()), /*ReturnValue*/TEXT(""), *GetterReturnType, *SetterParamDecl, bIsReadOnly);
+		}
+		else
+		{
+			// We use strict typing for return values to aid auto-complete (we also only care about the type and not the value, so structs can be default constructed)
+			static const PyGenUtil::EPythonizeFlags PythonizeValueFlags = PyGenUtil::EPythonizeFlags::UseStrictTyping | PyGenUtil::EPythonizeFlags::DefaultConstructStructs;
+			const FString GetterReturnValue = PyGenUtil::PythonizeDefaultValue(InGetSet.Prop.Prop, FString(), PythonizeValueFlags);
+			ExportGetSet(UTF8_TO_TCHAR(InGetSet.GetSetName.GetData()), UTF8_TO_TCHAR(InGetSet.GetSetDoc.GetData()), *GetterReturnValue, /*ReturnType*/TEXT(""), *SetterParamDecl, bIsReadOnly);
+		}
 	};
 
-	auto ExportGeneratedOperator = [&OutPythonScript](const PyGenUtil::FGeneratedWrappedOperatorStack& InOpStack, const PyGenUtil::FGeneratedWrappedOperatorSignature& InOpSignature)
+	auto ExportGeneratedOperator = [&OutPythonScript, &PyTypeName](const PyGenUtil::FGeneratedWrappedOperatorStack& InOpStack, const PyGenUtil::FGeneratedWrappedOperatorSignature& InOpSignature)
 	{
 		auto AppendFunctionTooltip = [](const UFunction* InFunc, const TCHAR* InIdentation, FString& OutStr)
 		{
@@ -2788,10 +3184,29 @@ void FPyWrapperTypeRegistry::GenerateStubCodeForWrappedType(PyTypeObject* PyType
 			AppendFunctionTooltip(InOpStack.Funcs[0].Func, TEXT(""), OpDocString);
 		}
 
-		OutPythonScript.WriteLine(FString::Printf(TEXT("def %s(self%s):"), InOpSignature.PyFuncName, (InOpSignature.OtherType == PyGenUtil::FGeneratedWrappedOperatorSignature::EType::None ? TEXT("") : TEXT(", other"))));
+		// Python complains if __eq__() and __ne__() overriding the 'object' methods don't have the other parameter as 'object' type.
+		FString OtherParamType = (FCString::Strcmp(InOpSignature.PyFuncName, TEXT("__eq__")) == 0 || FCString::Strcmp(InOpSignature.PyFuncName, TEXT("__ne__")) == 0) ? TEXT("object") : PyTypeName;
+
+		PyGenUtil::IsTypeHintingEnabled()
+			? OutPythonScript.WriteLine(FString::Printf(TEXT("def %s(self%s) -> %s:"),
+				InOpSignature.PyFuncName,
+				(InOpSignature.OtherType == PyGenUtil::FGeneratedWrappedOperatorSignature::EType::None ? TEXT("") : *FString::Printf(TEXT(", other: %s"),*OtherParamType)),
+				(InOpSignature.ReturnType == PyGenUtil::FGeneratedWrappedOperatorSignature::EType::Bool ? TEXT("bool") : TEXT("None"))))
+			: OutPythonScript.WriteLine(FString::Printf(TEXT("def %s(self%s):"), InOpSignature.PyFuncName, (InOpSignature.OtherType == PyGenUtil::FGeneratedWrappedOperatorSignature::EType::None ? TEXT("") : TEXT(", other"))));
 		OutPythonScript.IncreaseIndent();
 		OutPythonScript.WriteDocString(OpDocString);
-		OutPythonScript.WriteLine(InOpSignature.ReturnType == PyGenUtil::FGeneratedWrappedOperatorSignature::EType::Bool ? TEXT("return False") : TEXT("return None"));
+
+		PyGenUtil::IsTypeHintingEnabled()
+			? OutPythonScript.WriteLine(TEXT("...")) // No need to add a return value to help auto-complete/type checker when the type is hinted in the signature.
+			: OutPythonScript.WriteLine(InOpSignature.ReturnType == PyGenUtil::FGeneratedWrappedOperatorSignature::EType::Bool ? TEXT("return False") : TEXT("return None"));
+		OutPythonScript.DecreaseIndent();
+	};
+
+	auto ExportNoBodyMethod = [&OutPythonScript](const TCHAR* MethodDecl)
+	{
+		OutPythonScript.WriteLine(MethodDecl);
+		OutPythonScript.IncreaseIndent();
+		OutPythonScript.WriteLine(TEXT("...")); // pass
 		OutPythonScript.DecreaseIndent();
 	};
 
@@ -2814,9 +3229,6 @@ void FPyWrapperTypeRegistry::GenerateStubCodeForWrappedType(PyTypeObject* PyType
 			{
 				TSharedPtr<const FPyWrapperStructMetaData> StructMetaData = StaticCastSharedPtr<FPyWrapperStructMetaData>(GeneratedTypeData->MetaData);
 				check(StructMetaData.IsValid());
-				
-				// Don't export FDateTime values for struct __init__ as they can be non-deterministic
-				static const uint32 PythonizeValueFlags = PyGenUtil::EPythonizeValueFlags::DefaultConstructDateTime;
 
 				// Python can only support 255 parameters, so if we have more than that for this struct just use the generic __init__ function
 				FString InitParamsStr;
@@ -2827,12 +3239,7 @@ void FPyWrapperTypeRegistry::GenerateStubCodeForWrappedType(PyTypeObject* PyType
 					for (const PyGenUtil::FGeneratedWrappedMethodParameter& InitParam : StructMetaData->MakeFunc.InputParams)
 					{
 						InitParamsStr += TEXT(", ");
-						InitParamsStr += UTF8_TO_TCHAR(InitParam.ParamName.GetData());
-						if (InitParam.ParamDefaultValue.IsSet())
-						{
-							InitParamsStr += TEXT("=");
-							InitParamsStr += PyGenUtil::PythonizeDefaultValue(InitParam.ParamProp, InitParam.ParamDefaultValue.GetValue(), PythonizeValueFlags);
-						}
+						InitParamsStr += PyGenUtil::PythonizeMethodParam(InitParam, PythonizeMethodParamFlags, UE::Python::ShouldEllipseParamDefaultValue);
 					}
 				}
 				else if (StructMetaData->InitParams.Num() <= 255)
@@ -2843,23 +3250,16 @@ void FPyWrapperTypeRegistry::GenerateStubCodeForWrappedType(PyTypeObject* PyType
 					for (const PyGenUtil::FGeneratedWrappedMethodParameter& InitParam : StructMetaData->InitParams)
 					{
 						InitParamsStr += TEXT(", ");
-						InitParamsStr += UTF8_TO_TCHAR(InitParam.ParamName.GetData());
-						if (InitParam.ParamDefaultValue.IsSet())
-						{
-							InitParamsStr += TEXT("=");
-							InitParamsStr += PyGenUtil::PythonizeValue(InitParam.ParamProp, InitParam.ParamProp->ContainerPtrToValuePtr<void>(StructData.GetStructMemory()), PythonizeValueFlags);
-						}
+						InitParamsStr += PyGenUtil::PythonizeMethodParam(InitParam, PythonizeMethodParamFlags, UE::Python::ShouldEllipseParamDefaultValue);
 					}
 				}
 
 				if (!bWriteDefaultInit)
 				{
 					bHasExportedClassData = true;
-
-					OutPythonScript.WriteLine(FString::Printf(TEXT("def __init__(self%s):"), *InitParamsStr));
-					OutPythonScript.IncreaseIndent();
-					OutPythonScript.WriteLine(TEXT("pass"));
-					OutPythonScript.DecreaseIndent();
+					PyGenUtil::IsTypeHintingEnabled()
+						? ExportNoBodyMethod(*FString::Printf(TEXT("def __init__(self%s) -> None:"), *InitParamsStr))
+						: ExportNoBodyMethod(*FString::Printf(TEXT("def __init__(self%s):"), *InitParamsStr));
 				}
 			}
 			else if (MetaGuid == FPyWrapperEnumMetaData::StaticTypeId())
@@ -2871,31 +3271,31 @@ void FPyWrapperTypeRegistry::GenerateStubCodeForWrappedType(PyTypeObject* PyType
 		}
 		else if (PyType == &PyWrapperObjectType)
 		{
-			bWriteDefaultInit = false;
 			bHasExportedClassData = true;
+			bWriteDefaultInit = false;
 
-			OutPythonScript.WriteLine(TEXT("def __init__(self, outer=None, name=\"None\"):"));
-			OutPythonScript.IncreaseIndent();
-			OutPythonScript.WriteLine(TEXT("pass"));
-			OutPythonScript.DecreaseIndent();
+			PyGenUtil::IsTypeHintingEnabled()
+				? ExportNoBodyMethod(TEXT("def __init__(self, outer: Optional[Object] = None, name: Union[Name, str]=\"None\") -> None:"))
+				: ExportNoBodyMethod(TEXT("def __init__(self, outer=None, name=\"None\"):"));
 		}
 		else if (PyType == &PyWrapperEnumType)
 		{
 			// Enums don't really have an __init__ function at runtime, so just give them a default one (with no arguments)
+			bHasExportedClassData = true;
 			bWriteDefaultInit = false;
-
-			OutPythonScript.WriteLine(TEXT("def __init__(self):"));
-			OutPythonScript.IncreaseIndent();
-			OutPythonScript.WriteLine(TEXT("pass"));
-			OutPythonScript.DecreaseIndent();
+			PyGenUtil::IsTypeHintingEnabled()
+				? ExportNoBodyMethod(TEXT("def __init__(self) -> None:"))
+				: ExportNoBodyMethod(TEXT("def __init__(self):"));
 		}
 		else if (PyType == &PyWrapperEnumValueDescrType)
 		{
-			bWriteDefaultInit = false;
 			bHasExportedClassData = true;
+			bWriteDefaultInit = false;
 
 			// This is a special internal decorator type used to define enum entries, which is why it has __get__ as well as __init__
-			OutPythonScript.WriteLine(TEXT("def __init__(self, enum, name, value):"));
+			PyGenUtil::IsTypeHintingEnabled()
+				? OutPythonScript.WriteLine(TEXT("def __init__(self, enum: Any, name: str, value: int) -> None:"))
+				: OutPythonScript.WriteLine(TEXT("def __init__(self, enum, name, value):"));
 			OutPythonScript.IncreaseIndent();
 			OutPythonScript.WriteLine(TEXT("self.enum = enum"));
 			OutPythonScript.WriteLine(TEXT("self.name = name"));
@@ -2914,13 +3314,136 @@ void FPyWrapperTypeRegistry::GenerateStubCodeForWrappedType(PyTypeObject* PyType
 			OutPythonScript.DecreaseIndent();
 		}
 
+		if (PyGenUtil::IsTypeHintingEnabled())
+		{
+			if (PyType == &PyWrapperArrayType || PyType == &PyWrapperFixedArrayType)
+			{
+				bWriteDefaultInit = false;
+				bHasExportedClassData = true;
+
+				// WARNING: The parameter name used with __init__ must be 'type' and 'len' to allow named parameters.
+				PyType == &PyWrapperFixedArrayType
+					? ExportNoBodyMethod(TEXT("def __init__(self, type: type, len: int) -> None:"))
+					: ExportNoBodyMethod(TEXT("def __init__(self, type: type) -> None:"));
+
+				ExportNoBodyMethod(TEXT("def __setitem__(self, index: int, value: _ElemType) -> None:"));
+				ExportNoBodyMethod(TEXT("def __getitem__(self, index: int) -> _ElemType:"));
+			}
+			else if (PyType == &PyWrapperSetType)
+			{
+				// WARNING: The parameter name used with __init__ must be 'type' to allow named parameters.
+				ExportNoBodyMethod(TEXT("def __init__(self, type: type) -> None:"));
+				bWriteDefaultInit = false;
+				bHasExportedClassData = true;
+			}
+			else if (PyType == &PyWrapperMapType)
+			{
+				// WARNING: The parameter name used with __init__ must be 'key' and 'value' to allow named parameters.
+				ExportNoBodyMethod(TEXT("def __init__(self, key: type, value: type) -> None:"));
+				bWriteDefaultInit = false;
+
+				ExportNoBodyMethod(TEXT("def __setitem__(self, key: _KeyType, value: _ValueType) -> None:"));
+				ExportNoBodyMethod(TEXT("def __getitem__(self, key: _KeyType) -> _ValueType:"));
+			}
+			else if (PyType == &PyWrapperNameType)
+			{
+				ExportNoBodyMethod(TEXT("def __init__(self, value: Union[Name, str] = \"None\") -> None:"));
+				bWriteDefaultInit = false;
+				bHasExportedClassData = true;
+			}
+			else if (PyType == &PyWrapperTextType)
+			{
+				ExportNoBodyMethod(TEXT("def __init__(self, value: Union[Text, str] = \"\") -> None:"));
+				bWriteDefaultInit = false;
+				bHasExportedClassData = true;
+			}
+			else if (PyType == &PyWrapperFieldPathType)
+			{
+				ExportNoBodyMethod(TEXT("def __init__(self, value: Union[FieldPath, str] = \"\") -> None:"));
+				bWriteDefaultInit = false;
+				bHasExportedClassData = true;
+			}
+			else if (PyType == &PyScopedSlowTaskType) // Type defined in PyCore.cpp
+			{
+				// WARNING: Parameter names must match the ones from PyCore.cpp to allow named parameters.
+				ExportNoBodyMethod(TEXT("def __init__(self, work: float, desc: Union[Text, str] = \"\", enabled: bool = True) -> None:"));
+				bWriteDefaultInit = false;
+				bHasExportedClassData = true;
+			}
+			else if (PyType == &PyScopedEditorTransactionType) // Type defined in PyEditor.cpp
+			{
+				// WARNING: Parameter names must match the ones from PyEditor.cpp to allow named parameters.
+				ExportNoBodyMethod(TEXT("def __init__(self, desc: Union[Text, str]) -> None:"));
+				bWriteDefaultInit = false;
+				bHasExportedClassData = true;
+			}
+			else if (PyType == &PyActorIteratorType || PyType == &PySelectedActorIteratorType) // Type defined in PyEngine.cpp
+			{
+				// WARNING: Parameter names must match the ones from PyEngine.cpp to allow named parameters.
+				ExportNoBodyMethod(TEXT("def __init__(self, world: World, type: Union[Class, type] = ...) -> None:"));
+				ExportNoBodyMethod(TEXT("def __iter__(self) -> Iterator[Any]:"));
+				bWriteDefaultInit = false;
+				bHasExportedClassData = true;
+			}
+			else if (PyType == &PyObjectIteratorType)
+			{
+				// WARNING: Parameter names must match the ones from PyEngine.cpp to allow named parameters.
+				ExportNoBodyMethod(TEXT("def __init__(self, type: Union[Class, type] = ...) -> None:"));
+				ExportNoBodyMethod(TEXT("def __iter__(self) -> Iterator[Any]:"));
+				bWriteDefaultInit = false;
+				bHasExportedClassData = true;
+			}
+			else if (PyType == &PyTypeIteratorType || PyType == &PyClassIteratorType || PyType == &PyStructIteratorType)
+			{
+				// WARNING: Parameter names must match the ones from PyEngine.cpp to allow named parameters.
+				ExportNoBodyMethod(TEXT("def __init__(self, type: type) -> None:"));
+				if (PyType == &PyClassIteratorType)
+				{
+					ExportNoBodyMethod(TEXT("def __iter__(self) -> Iterator[Class]:"));
+				}
+				else if (PyType == &PyStructIteratorType)
+				{
+					ExportNoBodyMethod(TEXT("def __iter__(self) -> Iterator[ScriptStruct]:"));
+				}
+				else// if (PyType == &PyTypeIteratorType) -> The last remaining.
+				{
+					ExportNoBodyMethod(TEXT("def __iter__(self) -> Iterator[type]:"));
+				}
+				bWriteDefaultInit = false;
+				bHasExportedClassData = true;
+			}
+			else if (PyType == &PyFPropertyDefType)
+			{
+				// WARNING: Parameter names must match the ones from PyEngine.cpp to allow named parameters.
+				ExportNoBodyMethod(TEXT("def __init__(self, type: type, meta: Optional[Dict[str, str]], getter: Optional[str], setter: Optional[str]) -> None:"));
+				bWriteDefaultInit = false;
+				bHasExportedClassData = true;
+			}
+			else if (PyType == &PyUFunctionDefType)
+			{
+				// WARNING: Parameter names must match the ones from PyEngine.cpp to allow named parameters.
+				ExportNoBodyMethod(TEXT("def __init__(self, func: Any, meta: Optional[Dict[str, str]], ret: Optional[Any], params: Optional[Sequence[Any]], override: Optional[bool], static: Optional[bool], pure: Optional[bool], getter: Optional[bool], setter: Optional[bool]) -> None:"));
+				bWriteDefaultInit = false;
+				bHasExportedClassData = true;
+			}
+			else if (PyType == &PyUValueDefType)
+			{
+				// WARNING: Parameter names must match the ones from PyEngine.cpp to allow named parameters.
+				ExportNoBodyMethod(TEXT("def __init__(self, val: object, meta: Optional[Dict[str, str]]) -> None:"));
+				bWriteDefaultInit = false;
+				bHasExportedClassData = true;
+			}
+		}
+
 		if (bWriteDefaultInit)
 		{
 			bHasExportedClassData = true;
-			
-			OutPythonScript.WriteLine(TEXT("def __init__(self, *args, **kwargs):"));
+
+			PyGenUtil::IsTypeHintingEnabled()
+				? OutPythonScript.WriteLine(TEXT("def __init__(self, *args: Any, **kwargs: Any) -> None:"))
+				: OutPythonScript.WriteLine(TEXT("def __init__(self, *args, **kwargs):"));
 			OutPythonScript.IncreaseIndent();
-			OutPythonScript.WriteLine(TEXT("pass"));
+			OutPythonScript.WriteLine(TEXT("...")); // pass
 			OutPythonScript.DecreaseIndent();
 		}
 	}
@@ -2969,7 +3492,7 @@ void FPyWrapperTypeRegistry::GenerateStubCodeForWrappedType(PyTypeObject* PyType
 
 		bHasExportedClassData = true;
 
-		ExportGetSet(UTF8_TO_TCHAR(GetSetDef->name), UTF8_TO_TCHAR(GetSetDef->doc), TEXT("None"), /*IsReadOnly*/false);
+		ExportGetSet(UTF8_TO_TCHAR(GetSetDef->name), UTF8_TO_TCHAR(GetSetDef->doc), /*ReturnValue*/TEXT("None"), /*GetterReturnTypeDecl*/TEXT(""), /*ParamNameAndTypeDecl*/TEXT(""), /*IsReadOnly*/false);
 	}
 
 	for (PyMethodDef* MethodDef = PyType->tp_methods; MethodDef && MethodDef->ml_name; ++MethodDef)
@@ -2981,15 +3504,15 @@ void FPyWrapperTypeRegistry::GenerateStubCodeForWrappedType(PyTypeObject* PyType
 		if (bIsClassMethod)
 		{
 			OutPythonScript.WriteLine(TEXT("@classmethod"));
-			OutPythonScript.WriteLine(FString::Printf(TEXT("def %s(cls, *args%s):"), UTF8_TO_TCHAR(MethodDef->ml_name), bHasKeywords ? TEXT(", **kwargs") : TEXT("")));
+			OutPythonScript.WriteLine(UE::Python::GenerateMethodDecl(MethodDef, PyTypeName, /*bIsClassMethod*/true));
 		}
 		else
 		{
-			OutPythonScript.WriteLine(FString::Printf(TEXT("def %s(self, *args%s):"), UTF8_TO_TCHAR(MethodDef->ml_name), bHasKeywords ? TEXT(", **kwargs") : TEXT("")));
+			OutPythonScript.WriteLine(UE::Python::GenerateMethodDecl(MethodDef, PyTypeName, /*bIsClassMethod*/false, /*bIsMemberMethod*/true));
 		}
 		OutPythonScript.IncreaseIndent();
 		OutPythonScript.WriteDocString(UTF8_TO_TCHAR(MethodDef->ml_doc));
-		OutPythonScript.WriteLine(TEXT("pass"));
+		OutPythonScript.WriteLine(TEXT("...")); // pass
 		OutPythonScript.DecreaseIndent();
 	}
 
@@ -3040,13 +3563,24 @@ void FPyWrapperTypeRegistry::GenerateStubCodeForWrappedType(PyTypeObject* PyType
 				ExportGeneratedMethod(*DynamicMethod);
 			}
 
+			// Avoid exporting method __truediv__ twice becauses it is used for both / and /= operator.
+			bool bTrueDivExported = false;
 			for (int32 OpTypeIndex = 0; OpTypeIndex < (int32)PyGenUtil::EGeneratedWrappedOperatorType::Num; ++OpTypeIndex)
 			{
 				const PyGenUtil::FGeneratedWrappedOperatorStack& OpStack = StructMetaData->OpStacks[OpTypeIndex];
 				if (OpStack.Funcs.Num() > 0)
 				{
 					const PyGenUtil::FGeneratedWrappedOperatorSignature& OpSignature = PyGenUtil::FGeneratedWrappedOperatorSignature::OpTypeToSignature((PyGenUtil::EGeneratedWrappedOperatorType)OpTypeIndex);
-					ExportGeneratedOperator(OpStack, OpSignature);
+
+					if (FCString::Strcmp(OpSignature.PyFuncName, TEXT("__truediv__")) != 0)
+					{
+						ExportGeneratedOperator(OpStack, OpSignature);
+					}
+					else if (!bTrueDivExported)
+					{
+						ExportGeneratedOperator(OpStack, OpSignature);
+						bTrueDivExported = true;
+					}
 				}
 			}
 
@@ -3081,7 +3615,10 @@ void FPyWrapperTypeRegistry::GenerateStubCodeForWrappedType(PyTypeObject* PyType
 						EntryDoc.InsertAt(0, *FString::Printf(TEXT("%s: "), *EntryValue));
 					}
 
-					ExportConstantValue(*EntryName, *EntryDoc, FString::Printf(TEXT("_EnumEntry(\"%s\", \"%s\", %s)"), *PyTypeName, *EntryName, *EntryValue));
+					// The enum value types reported by the engine using print(type(unreal.SomeEnum.SOME_VALUE)) would be 'class SomeEnum' but _EnumEntry class is unrelated.
+					PyGenUtil::IsTypeHintingEnabled()
+						? ExportConstant(*EntryName, PyTypeName, /*Value*/FString(), *EntryDoc)
+						: ExportConstant(*EntryName, /*TypeHint*/FString(), FString::Printf(TEXT("_EnumEntry(\"%s\", \"%s\", %s)"), *PyTypeName, *EntryName, *EntryValue), *EntryDoc);
 				}
 			}
 		}
@@ -3089,7 +3626,7 @@ void FPyWrapperTypeRegistry::GenerateStubCodeForWrappedType(PyTypeObject* PyType
 
 	if (!bHasExportedClassData)
 	{
-		OutPythonScript.WriteLine(TEXT("pass"));
+		OutPythonScript.WriteLine(TEXT("...")); // pass
 	}
 
 	OutPythonScript.DecreaseIndent();
diff --git a/Engine/Plugins/Experimental/PythonScriptPlugin/Source/PythonScriptPlugin/Private/PythonOnlineDocsCommandlet.cpp b/Engine/Plugins/Experimental/PythonScriptPlugin/Source/PythonScriptPlugin/Private/PythonOnlineDocsCommandlet.cpp
--- a/Engine/Plugins/Experimental/PythonScriptPlugin/Source/PythonScriptPlugin/Private/PythonOnlineDocsCommandlet.cpp
+++ b/Engine/Plugins/Experimental/PythonScriptPlugin/Source/PythonScriptPlugin/Private/PythonOnlineDocsCommandlet.cpp
@@ -2,6 +2,7 @@
 
 #include "PythonOnlineDocsCommandlet.h"
 #include "PyWrapperTypeRegistry.h"
+#include "PyGenUtil.h"
 
 DEFINE_LOG_CATEGORY_STATIC(LogPythonOnlineDocsCommandlet, Log, All);
 
@@ -42,6 +43,9 @@ int32 UPythonOnlineDocsCommandlet::Main(const FString& Params)
 		DocGenFlags = EPyOnlineDocsFilterFlags::IncludeAll;
 	}
 
+	// For the documentation, use the the light typing version, avoid listing all type coercions which make the API harder to read.
+	PyGenUtil::SetTypeHintingMode(ETypeHintingMode::AutoCompletion);
+
 	UE_LOG(LogPythonOnlineDocsCommandlet, Display, TEXT("\n\nGenerating Python documentation..."));
 	FPyWrapperTypeRegistry::Get().GenerateStubCodeForWrappedTypes(DocGenFlags);
 
diff --git a/Engine/Plugins/Experimental/PythonScriptPlugin/Source/PythonScriptPlugin/Private/PythonScriptPlugin.cpp b/Engine/Plugins/Experimental/PythonScriptPlugin/Source/PythonScriptPlugin/Private/PythonScriptPlugin.cpp
--- a/Engine/Plugins/Experimental/PythonScriptPlugin/Source/PythonScriptPlugin/Private/PythonScriptPlugin.cpp
+++ b/Engine/Plugins/Experimental/PythonScriptPlugin/Source/PythonScriptPlugin/Private/PythonScriptPlugin.cpp
@@ -817,6 +817,9 @@ void FPythonScriptPlugin::InitializePython()
 			PyUtil::EnableDeveloperWarnings();
 		}
 
+		// Check if the user wants type hinting. (In the stub and/or Docstrings).
+		PyGenUtil::SetTypeHintingMode(GetTypeHintingMode());
+
 		// Initialize our custom method type as we'll need it when generating bindings
 		InitializePyMethodWithClosure();
 
@@ -1540,6 +1543,11 @@ bool FPythonScriptPlugin::IsDeveloperModeEnabled()
 	return GetDefault<UPythonScriptPluginSettings>()->bDeveloperMode || GetDefault<UPythonScriptPluginUserSettings>()->bDeveloperMode;
 }
 
+ETypeHintingMode FPythonScriptPlugin::GetTypeHintingMode()
+{
+	return GetDefault<UPythonScriptPluginUserSettings>()->TypeHintingMode;
+}
+
 void FPythonScriptPlugin::OnAssetRenamed(const FAssetData& Data, const FString& OldName)
 {
 	TRACE_CPUPROFILER_EVENT_SCOPE(FPythonScriptPlugin::OnAssetRenamed)
diff --git a/Engine/Plugins/Experimental/PythonScriptPlugin/Source/PythonScriptPlugin/Private/PythonScriptPlugin.h b/Engine/Plugins/Experimental/PythonScriptPlugin/Source/PythonScriptPlugin/Private/PythonScriptPlugin.h
--- a/Engine/Plugins/Experimental/PythonScriptPlugin/Source/PythonScriptPlugin/Private/PythonScriptPlugin.h
+++ b/Engine/Plugins/Experimental/PythonScriptPlugin/Source/PythonScriptPlugin/Private/PythonScriptPlugin.h
@@ -12,6 +12,7 @@
 #include "Framework/Commands/InputChord.h"
 #include "Kismet2/EnumEditorUtils.h"
 #include "UObject/StrongObjectPtr.h"
+#include "PythonScriptPluginSettings.h"
 
 class FPythonScriptPlugin;
 class FPythonScriptRemoteExecution;
@@ -176,6 +177,8 @@ private:
 
 	static bool IsDeveloperModeEnabled();
 
+	static ETypeHintingMode GetTypeHintingMode();
+
 	void OnAssetRenamed(const FAssetData& Data, const FString& OldName);
 
 	void OnAssetRemoved(const FAssetData& Data);
diff --git a/Engine/Plugins/Experimental/PythonScriptPlugin/Source/PythonScriptPlugin/Private/PythonScriptPluginSettings.h b/Engine/Plugins/Experimental/PythonScriptPlugin/Source/PythonScriptPlugin/Private/PythonScriptPluginSettings.h
--- a/Engine/Plugins/Experimental/PythonScriptPlugin/Source/PythonScriptPlugin/Private/PythonScriptPluginSettings.h
+++ b/Engine/Plugins/Experimental/PythonScriptPlugin/Source/PythonScriptPlugin/Private/PythonScriptPluginSettings.h
@@ -7,6 +7,28 @@
 #include "Engine/DeveloperSettings.h"
 #include "PythonScriptPluginSettings.generated.h"
 
+UENUM()
+enum class ETypeHintingMode : uint8
+{
+	/** Turn off type hinting. */
+	Off,
+
+	/**
+	 * When generating the Python stub and to some extend the Docstrings, enables type hinting (PEP 484) to get the best experience
+	 * with a Python IDE auto-completion. The hinting will list the exact input types, omit type coercions and will assume all reflected
+	 * unreal.Object cannot be None which is not true, but will let the function signature easy to read.
+	 */
+	AutoCompletion,
+
+	/**
+	 * Enables type hinting for static type checking. Hint as close as possible the real supported types including
+	 * possible type coercions. Because the UE reflection API doesn't provide all the required information, some tradeoffs
+	 * are required that do not always reflect the reality. For example, reflected UObject will always be marked as
+	 * 'possibly None'. While this is true in some contexts, it is not true all the time.
+	 */
+	TypeChecker,
+};
+
 /**
  * Configure the Python plug-in.
  */
@@ -103,6 +125,15 @@ public:
 	UPROPERTY(config, EditAnywhere, Category=Python, meta=(ConfigRestartRequired=true))
 	bool bDeveloperMode;
 
+	/**
+	 * Should the generated Python stub and API documentation have type hints. This enables standard Python type hinting (PEP 484) for the classes,
+	 * structs, methods, properties, constants, etc. exposed by the engine. If the developer mode is enabled and the Python IDE configured to use
+	 * the generated Python stub, types will be displayed in auto-completion popup and used by the IDE static type checkers. This has no effects on
+	 * the execution of the code. (Requires Python >= 3.7)
+	 */
+	UPROPERTY(config, EditAnywhere, Category=Python, meta=(ConfigRestartRequired=true))
+	ETypeHintingMode TypeHintingMode = ETypeHintingMode::AutoCompletion;
+
 	/** Should Python scripts be available in the Content Browser? */
 	UPROPERTY(config, EditAnywhere, Category=Python, meta=(ConfigRestartRequired=true))
 	bool bEnableContentBrowserIntegration;
