Fixed CIS complaining about Native pointer usage in member declaration detected [[[UPyTestObject*]]].

diff --git a/Engine/Plugins/Experimental/PythonScriptPlugin/Source/PythonScriptPlugin/Private/PyTest.h b/Engine/Plugins/Experimental/PythonScriptPlugin/Source/PythonScriptPlugin/Private/PyTest.h
--- a/Engine/Plugins/Experimental/PythonScriptPlugin/Source/PythonScriptPlugin/Private/PyTest.h
+++ b/Engine/Plugins/Experimental/PythonScriptPlugin/Source/PythonScriptPlugin/Private/PyTest.h
@@ -5,6 +5,7 @@
 #include "CoreMinimal.h"
 #include "PyTestInterface.h"
 #include "Kismet/BlueprintFunctionLibrary.h"
+#include "UObject/ObjectPtr.h"
 #include "PyTest.generated.h"
 
 /**
@@ -387,7 +388,7 @@ public:
 	FPyTestStruct StructProp;
 
 	UPROPERTY(EditAnywhere, BlueprintReadWrite, Category = "Python|Internal")
-	UPyTestObject* ObjectProp = nullptr;
+	TObjectPtr<UPyTestObject> ObjectProp = nullptr;
 
 	UPROPERTY(EditAnywhere, BlueprintReadWrite, Category = "Python|Internal")
 	TArray<FString> StrArrayProp;
@@ -399,7 +400,7 @@ public:
 	TArray<FText> TextArrayProp;
 
 	UPROPERTY(EditAnywhere, BlueprintReadWrite, Category = "Python|Internal")
-	TArray<UObject*> ObjectArrayProp;
+	TArray<TObjectPtr<UObject>> ObjectArrayProp;
 
 	UPROPERTY(EditAnywhere, BlueprintReadWrite, Category = "Python|Internal")
 	TSet<FString> SetProp;
