Fix CIS static analytis in PyGetUtil.cpp

diff --git a/Engine/Plugins/Experimental/PythonScriptPlugin/Source/PythonScriptPlugin/Private/PyGenUtil.cpp b/Engine/Plugins/Experimental/PythonScriptPlugin/Source/PythonScriptPlugin/Private/PyGenUtil.cpp
--- a/Engine/Plugins/Experimental/PythonScriptPlugin/Source/PythonScriptPlugin/Private/PyGenUtil.cpp
+++ b/Engine/Plugins/Experimental/PythonScriptPlugin/Source/PythonScriptPlugin/Private/PyGenUtil.cpp
@@ -3368,7 +3368,11 @@
 {
 	// In 3.7, a feature allowing usage of non-declared types in type hints was added (from __future__ import annotations) and the engine
 	// doesn't output types in a strict declaration/usage order, so this is required to correctly type.
-	return PY_MAJOR_VERSION >= 3 && PY_MINOR_VERSION >= 7 && GTypeHintingMode != ETypeHintingMode::Off;
+#if PY_MAJOR_VERSION >= 3 && PY_MINOR_VERSION >= 7
+	return GTypeHintingMode != ETypeHintingMode::Off;
+#else
+	return false;
+#endif
 }
 
 }
