[Insights] Added LLM_SCOPEs in PythonScriptPlugin (in an effort to reduce the Untagged allocations).

diff --git a/Engine/Plugins/Experimental/PythonScriptPlugin/Source/PythonScriptPlugin/Private/PythonScriptPlugin.cpp b/Engine/Plugins/Experimental/PythonScriptPlugin/Source/PythonScriptPlugin/Private/PythonScriptPlugin.cpp
--- a/Engine/Plugins/Experimental/PythonScriptPlugin/Source/PythonScriptPlugin/Private/PythonScriptPlugin.cpp
+++ b/Engine/Plugins/Experimental/PythonScriptPlugin/Source/PythonScriptPlugin/Private/PythonScriptPlugin.cpp
@@ -27,6 +27,7 @@
 #include "Misc/FileHelper.h"
 #include "Misc/PackageName.h"
 #include "HAL/FileManager.h"
+#include "HAL/LowLevelMemTracker.h"
 #include "HAL/PlatformMisc.h"
 #include "HAL/PlatformProcess.h"
 #include "Containers/Ticker.h"
@@ -604,6 +605,8 @@ void FPythonScriptPlugin::StartupModule()
 	}
 
 #if WITH_PYTHON
+	LLM_SCOPE_BYNAME(TEXT("PythonScriptPlugin"));
+
 	InitializePython();
 	IModularFeatures::Get().RegisterModularFeature(IConsoleCommandExecutor::ModularFeatureName(), &CmdExec);
 	IModularFeatures::Get().RegisterModularFeature(IConsoleCommandExecutor::ModularFeatureName(), &CmdREPLExec);
@@ -1446,6 +1449,7 @@ void FPythonScriptPlugin::OnModuleDirtied(FName InModuleName)
 
 void FPythonScriptPlugin::OnModulesChanged(FName InModuleName, EModuleChangeReason InModuleChangeReason)
 {
+	LLM_SCOPE_BYNAME(TEXT("PythonScriptPlugin"));
 	TRACE_CPUPROFILER_EVENT_SCOPE(FPythonScriptPlugin::OnModulesChanged)
 
 	switch (InModuleChangeReason)
