PythonScriptPlugin: ensure that custom excepthooks are called with all three arguments

Custom excepthooks can be installed in Python to be invoked when an exception is
raised:

https://docs.python.org/3/library/sys.html#sys.excepthook

The excepthook is expected to take three arguments: the exception class, an
exception instance, and a traceback object.

The PythonScriptPlugin has some handling to detect when a custom excepthook is
installed and make sure that it gets called, which it does from C++ using
PyObject_CallFunctionObjArgs(). Depending on the exception though, it's possible
that the FPyObjectPtrs that hold the exception instance and/or the traceback object
could be invalid. When passed to PyObject_CallFunctionObjArgs(), their value then
becomes nullptr, which is used to signal the end of the arguments list to the function
and would therefore cause too few arguments to be passed to the excepthook.
This would put the interpreter into a bad state where it would continue to attempt to
call the excepthook with too few arguments, and no future commands would work.
The easiest way to get into this bad state was to install a custom excepthook and
then issue a Python command with invalid syntax. Doing so raises a SyntaxError with
no traceback object, since the code could not actually be executed.

To address this, we explicitly pass None for the exception instance and/or the traceback
object when they are invalid to ensure that the excepthook always receives three
arguments.

diff --git a/Engine/Plugins/Experimental/PythonScriptPlugin/Source/PythonScriptPlugin/Private/PyUtil.cpp b/Engine/Plugins/Experimental/PythonScriptPlugin/Source/PythonScriptPlugin/Private/PyUtil.cpp
--- a/Engine/Plugins/Experimental/PythonScriptPlugin/Source/PythonScriptPlugin/Private/PyUtil.cpp
+++ b/Engine/Plugins/Experimental/PythonScriptPlugin/Source/PythonScriptPlugin/Private/PyUtil.cpp
@@ -1426,7 +1426,13 @@ bool FetchPythonError(FString& OutError)
 		PyObject* PyCurrentExceptHook = PySys_GetObject(PyCStrCast("excepthook"));
 		if (PyCurrentExceptHook && PyDefaultExceptHook && PyCurrentExceptHook != PyDefaultExceptHook)
 		{
-			FPyObjectPtr PyExceptHookResult = FPyObjectPtr::StealReference(PyObject_CallFunctionObjArgs(PyCurrentExceptHook, PyExceptionType.Get(), PyExceptionValue.Get(), PyExceptionTraceback.Get(), nullptr));
+			FPyObjectPtr PyExceptHookResult = FPyObjectPtr::StealReference(
+				PyObject_CallFunctionObjArgs(
+					PyCurrentExceptHook,
+					PyExceptionType.Get(),
+					PyExceptionValue ? PyExceptionValue.Get() : Py_None,
+					PyExceptionTraceback ? PyExceptionTraceback.Get() : Py_None,
+					nullptr));
 		}
 	}
 
