Fix Python GIL deadlock with reinstancing object.

diff --git a/Engine/Plugins/Experimental/PythonScriptPlugin/Source/PythonScriptPlugin/Private/PyWrapperTypeRegistry.cpp b/Engine/Plugins/Experimental/PythonScriptPlugin/Source/PythonScriptPlugin/Private/PyWrapperTypeRegistry.cpp
--- a/Engine/Plugins/Experimental/PythonScriptPlugin/Source/PythonScriptPlugin/Private/PyWrapperTypeRegistry.cpp
+++ b/Engine/Plugins/Experimental/PythonScriptPlugin/Source/PythonScriptPlugin/Private/PyWrapperTypeRegistry.cpp
@@ -781,12 +781,14 @@ void FPyWrapperTypeReinstancer::ProcessPending()
 			}
 		}
 
+		Py_BEGIN_ALLOW_THREADS
 		Reload->Reinstance();
 
 		ClassesToReinstance.Reset();
 		StructsToReinstance.Reset();
 
 		CollectGarbage(GARBAGE_COLLECTION_KEEPFLAGS);
+		Py_END_ALLOW_THREADS
 	}
 }
 
