Fixed python type hinting for unreal.get_editor_subsystem() and unreal.get_engine_subsystem() to hint the correct types.

Thanks to [REDACTED_PERSON] for this contribution.

diff --git a/Engine/Plugins/Experimental/PythonScriptPlugin/Source/PythonScriptPlugin/Private/PyEditor.cpp b/Engine/Plugins/Experimental/PythonScriptPlugin/Source/PythonScriptPlugin/Private/PyEditor.cpp
--- a/Engine/Plugins/Experimental/PythonScriptPlugin/Source/PythonScriptPlugin/Private/PyEditor.cpp
+++ b/Engine/Plugins/Experimental/PythonScriptPlugin/Source/PythonScriptPlugin/Private/PyEditor.cpp
@@ -202,8 +202,8 @@ static PyObject* GetEditorSubsystem(FPyScopedEditorTransaction* InSelf, PyObject
 }
 
 PyMethodDef PyEditorMethods[] = {
-	{ "get_engine_subsystem", PyCFunctionCast(&PyEditor::GetEngineSubsystem), METH_VARARGS, "get_engine_subsystem(subsystem : Union[Class, type]) -> Any -- returns the requested Engine subsystem or None" },
-	{ "get_editor_subsystem", PyCFunctionCast(&PyEditor::GetEditorSubsystem), METH_VARARGS, "get_editor_subsystem(subsystem : Union[Class, type]) -> Any -- returns the requested Editor subsystem or None" },
+	{ "get_engine_subsystem", PyCFunctionCast(&PyEditor::GetEngineSubsystem), METH_VARARGS, "get_engine_subsystem(subsystem : Union[Class, type[_EngineSubsystemTypeVar]]) -> Optional[_EngineSubsystemTypeVar] -- returns the requested Engine subsystem or None" },
+	{ "get_editor_subsystem", PyCFunctionCast(&PyEditor::GetEditorSubsystem), METH_VARARGS, "get_editor_subsystem(subsystem : Union[Class, type[_EditorSubsystemTypeVar]]) -> Optional[_EditorSubsystemTypeVar] -- returns the requested Editor subsystem or None" },
 	{ nullptr, nullptr, 0, nullptr }
 };
 
diff --git a/Engine/Plugins/Experimental/PythonScriptPlugin/Source/PythonScriptPlugin/Private/PyWrapperTypeRegistry.cpp b/Engine/Plugins/Experimental/PythonScriptPlugin/Source/PythonScriptPlugin/Private/PyWrapperTypeRegistry.cpp
--- a/Engine/Plugins/Experimental/PythonScriptPlugin/Source/PythonScriptPlugin/Private/PyWrapperTypeRegistry.cpp
+++ b/Engine/Plugins/Experimental/PythonScriptPlugin/Source/PythonScriptPlugin/Private/PyWrapperTypeRegistry.cpp
@@ -2710,6 +2710,9 @@ void FPyWrapperTypeRegistry::GenerateStubCodeForWrappedTypes(const EPyOnlineDocs
 			PythonScript.WriteLine(TEXT("_ElemType = TypeVar('_ElemType')"));   // For Array, FixedArray, Set classes along with casting operations.
 			PythonScript.WriteLine(TEXT("_KeyType = TypeVar('_KeyType')"));     // For Mapping[KV, KT] used by Map class.
 			PythonScript.WriteLine(TEXT("_ValueType = TypeVar('_ValueType')")); // For Mapping[KV, KT] used by Map class.
+			PythonScript.WriteLine(TEXT("_ValueType = TypeVar('_ValueType')")); // For Mapping[KV, KT] used by Map class.
+			PythonScript.WriteLine(TEXT("_EngineSubsystemTypeVar = TypeVar('_EngineSubsystemTypeVar', bound=EngineSubsystem)"));
+			PythonScript.WriteLine(TEXT("_EditorSubsystemTypeVar  = TypeVar('_EditorSubsystemTypeVar', bound=EditorSubsystem)"));
 			PythonScript.WriteNewLine();
 		}
 	}
