Fixed a GIL deadlock with unreal.load_class(...) to load a BPGC

diff --git a/Engine/Plugins/Experimental/PythonScriptPlugin/Source/PythonScriptPlugin/Private/PyCore.cpp b/Engine/Plugins/Experimental/PythonScriptPlugin/Source/PythonScriptPlugin/Private/PyCore.cpp
--- a/Engine/Plugins/Experimental/PythonScriptPlugin/Source/PythonScriptPlugin/Private/PyCore.cpp
+++ b/Engine/Plugins/Experimental/PythonScriptPlugin/Source/PythonScriptPlugin/Private/PyCore.cpp
@@ -1339,7 +1339,10 @@ PyObject* LoadClass(PyObject* InSelf, PyObject* InArgs, PyObject* InKwds)
 		ObjectType = UObject::StaticClass();
 	}
 
-	UClass* PotentialClass = ::StaticLoadClass(ObjectType, ObjectOuter, *ObjectName);
+	UClass* PotentialClass = nullptr;
+	Py_BEGIN_ALLOW_THREADS
+	PotentialClass = ::StaticLoadClass(ObjectType, ObjectOuter, *ObjectName);
+	Py_END_ALLOW_THREADS
 	return PyConversion::PythonizeClass(PotentialClass);
 }
 
