USD: Fix editor hang caused by sneaking into C++ from Python without releasing the GIL.

diff --git a/Engine/Plugins/Experimental/PythonScriptPlugin/Source/PythonScriptPlugin/Private/PyConversion.cpp b/Engine/Plugins/Experimental/PythonScriptPlugin/Source/PythonScriptPlugin/Private/PyConversion.cpp
--- a/Engine/Plugins/Experimental/PythonScriptPlugin/Source/PythonScriptPlugin/Private/PyConversion.cpp
+++ b/Engine/Plugins/Experimental/PythonScriptPlugin/Source/PythonScriptPlugin/Private/PyConversion.cpp
@@ -1115,12 +1115,14 @@ FPyConversionResult PythonizeProperty_InContainer(const FProperty* Prop, const v
 
 void EmitPropertyChangeNotifications(const FPropertyAccessChangeNotify* InChangeNotify, const bool bIdenticalValue, const TFunctionRef<void()>& InDoChangeFunc)
 {
+	Py_BEGIN_ALLOW_THREADS
 	PropertyAccessUtil::EmitPreChangeNotify(InChangeNotify, bIdenticalValue);
 	if (!bIdenticalValue)
 	{
 		InDoChangeFunc();
 	}
 	PropertyAccessUtil::EmitPostChangeNotify(InChangeNotify, bIdenticalValue);
+	Py_END_ALLOW_THREADS
 }
 
 }
