Release the GIL before calling FPyWrapperTypeReinstancer::ProcessPending, not within it

Most places that call FPyWrapperTypeReinstancer::ProcessPending don't hold the GIL, so trying to release it inside just lead to an internal Python error

#rnx

diff --git a/Engine/Plugins/Experimental/PythonScriptPlugin/Source/PythonScriptPlugin/Private/PyWrapperTypeRegistry.cpp b/Engine/Plugins/Experimental/PythonScriptPlugin/Source/PythonScriptPlugin/Private/PyWrapperTypeRegistry.cpp
--- a/Engine/Plugins/Experimental/PythonScriptPlugin/Source/PythonScriptPlugin/Private/PyWrapperTypeRegistry.cpp
+++ b/Engine/Plugins/Experimental/PythonScriptPlugin/Source/PythonScriptPlugin/Private/PyWrapperTypeRegistry.cpp
@@ -781,14 +781,12 @@ void FPyWrapperTypeReinstancer::ProcessPending()
 			}
 		}
 
-		Py_BEGIN_ALLOW_THREADS
 		Reload->Reinstance();
 
 		ClassesToReinstance.Reset();
 		StructsToReinstance.Reset();
 
 		CollectGarbage(GARBAGE_COLLECTION_KEEPFLAGS);
-		Py_END_ALLOW_THREADS
 	}
 }
 
diff --git a/Engine/Plugins/Experimental/PythonScriptPlugin/Source/PythonScriptPlugin/Private/PythonScriptLibrary.cpp b/Engine/Plugins/Experimental/PythonScriptPlugin/Source/PythonScriptPlugin/Private/PythonScriptLibrary.cpp
--- a/Engine/Plugins/Experimental/PythonScriptPlugin/Source/PythonScriptPlugin/Private/PythonScriptLibrary.cpp
+++ b/Engine/Plugins/Experimental/PythonScriptPlugin/Source/PythonScriptPlugin/Private/PythonScriptLibrary.cpp
@@ -118,7 +118,9 @@ DEFINE_FUNCTION(UPythonScriptLibrary::execExecutePythonScript)
 
 		// Execute the Python command
 		FPyObjectPtr PyResult = FPyObjectPtr::StealReference(FPythonScriptPlugin::Get()->EvalString(*PythonScript, TEXT("<string>"), Py_file_input, PyTempGlobalDict, PyTempGlobalDict));
+		Py_BEGIN_ALLOW_THREADS
 		FPyWrapperTypeReinstancer::Get().ProcessPending();
+		Py_END_ALLOW_THREADS
 
 		// Read the output values from the Python context
 		if (PyResult)
