Fixed Python doc generation failing becauses the bound types used in type hinting were not defined yet in the generated Python stub.

diff --git a/Engine/Plugins/Experimental/PythonScriptPlugin/Source/PythonScriptPlugin/Private/PyWrapperTypeRegistry.cpp b/Engine/Plugins/Experimental/PythonScriptPlugin/Source/PythonScriptPlugin/Private/PyWrapperTypeRegistry.cpp
--- a/Engine/Plugins/Experimental/PythonScriptPlugin/Source/PythonScriptPlugin/Private/PyWrapperTypeRegistry.cpp
+++ b/Engine/Plugins/Experimental/PythonScriptPlugin/Source/PythonScriptPlugin/Private/PyWrapperTypeRegistry.cpp
@@ -2709,8 +2709,8 @@ void FPyWrapperTypeRegistry::GenerateStubCodeForWrappedTypes(const EPyOnlineDocs
 			PythonScript.WriteLine(TEXT("_KeyType = TypeVar('_KeyType')"));     // For Mapping[KV, KT] used by Map class.
 			PythonScript.WriteLine(TEXT("_ValueType = TypeVar('_ValueType')")); // For Mapping[KV, KT] used by Map class.
 			PythonScript.WriteLine(TEXT("_ValueType = TypeVar('_ValueType')")); // For Mapping[KV, KT] used by Map class.
-			PythonScript.WriteLine(TEXT("_EngineSubsystemTypeVar = TypeVar('_EngineSubsystemTypeVar', bound=EngineSubsystem)"));
-			PythonScript.WriteLine(TEXT("_EditorSubsystemTypeVar  = TypeVar('_EditorSubsystemTypeVar', bound=EditorSubsystem)"));
+			PythonScript.WriteLine(TEXT("_EngineSubsystemTypeVar = TypeVar('_EngineSubsystemTypeVar', bound='EngineSubsystem')"));
+			PythonScript.WriteLine(TEXT("_EditorSubsystemTypeVar  = TypeVar('_EditorSubsystemTypeVar', bound='EditorSubsystem')"));
 			PythonScript.WriteNewLine();
 		}
 	}
