Added the __class_getitem__ to the PyWrapperArray, PyWrapperFixedArray, PyWrapperSet and PyWrapperMap to get type hinting working properly when container is used to hint an method input parameters or the return value (or simply declared globally).
  - For more information, see https://peps.python.org/pep-0560/
  - Added corresponding tests to test_type_hints.py.

diff --git a/Engine/Plugins/Experimental/PythonScriptPlugin/Source/PythonScriptPlugin/Private/PyWrapperArray.cpp b/Engine/Plugins/Experimental/PythonScriptPlugin/Source/PythonScriptPlugin/Private/PyWrapperArray.cpp
--- a/Engine/Plugins/Experimental/PythonScriptPlugin/Source/PythonScriptPlugin/Private/PyWrapperArray.cpp
+++ b/Engine/Plugins/Experimental/PythonScriptPlugin/Source/PythonScriptPlugin/Private/PyWrapperArray.cpp
@@ -1150,6 +1150,13 @@ PyTypeObject InitializePyWrapperArrayType()
 
 	struct FMethods
 	{
+		// This is for type hinting, to make the class appears as a generic (https://peps.python.org/pep-0560/) and support hinting like 'a: unreal.Array[int] = unreal.Array(int)'
+		static PyObject* GetClassItem(PyObject* Type, PyObject* Item)
+		{
+			Py_INCREF(Type);
+			return Type;
+		}
+
 		static PyObject* Cast(PyTypeObject* InType, PyObject* InArgs, PyObject* InKwds)
 		{
 			PyObject* PyTypeObj = nullptr;
@@ -1390,6 +1397,7 @@ PyTypeObject InitializePyWrapperArrayType()
 
 	// NOTE: _ElemType = typing.TypeVar('_ElemType') and Type/Optional/Iterable/Callable are defines by the Python typing module.
 	static PyMethodDef PyMethods[] = {
+		{ "__class_getitem__", PyCFunctionCast(FMethods::GetClassItem), METH_O|METH_CLASS, "__class_getitem__(cls, item: _ElemType) -- implemented for type hinting purpose only" },
 		{ "cast", PyCFunctionCast(&FMethods::Cast), METH_VARARGS | METH_KEYWORDS | METH_CLASS, "cast(cls, type: Type[_ElemType], obj: object) -> Array[_ElemType] -- cast the given object to this Unreal array type" },
 		{ "__copy__", PyCFunctionCast(&FMethods::Copy), METH_NOARGS, "__copy__(self) -> Array[_ElemType] -- copy this Unreal array" },
 		{ "copy", PyCFunctionCast(&FMethods::Copy), METH_NOARGS, "copy(self) -> Array[_ElemType] -- copy this Unreal array" },
diff --git a/Engine/Plugins/Experimental/PythonScriptPlugin/Source/PythonScriptPlugin/Private/PyWrapperFixedArray.cpp b/Engine/Plugins/Experimental/PythonScriptPlugin/Source/PythonScriptPlugin/Private/PyWrapperFixedArray.cpp
--- a/Engine/Plugins/Experimental/PythonScriptPlugin/Source/PythonScriptPlugin/Private/PyWrapperFixedArray.cpp
+++ b/Engine/Plugins/Experimental/PythonScriptPlugin/Source/PythonScriptPlugin/Private/PyWrapperFixedArray.cpp
@@ -758,6 +758,13 @@ PyTypeObject InitializePyWrapperFixedArrayType()
 
 	struct FMethods
 	{
+		// This is for type hinting, to make the class appears as a generic (https://peps.python.org/pep-0560/) and support hinting like 'a: unreal.FixedArray[int] = unreal.FixedArray(int)'
+		static PyObject* GetClassItem(PyObject* Type, PyObject* Item)
+		{
+			Py_INCREF(Type);
+			return Type;
+		}
+
 		static PyObject* Cast(PyTypeObject* InType, PyObject* InArgs, PyObject* InKwds)
 		{
 			PyObject* PyTypeObj = nullptr;
@@ -796,6 +803,7 @@ PyTypeObject InitializePyWrapperFixedArrayType()
 
 	// NOTE: _ElemType = typing.TypeVar('_ElemType') and TYpe is defined in the typing Python module.
 	static PyMethodDef PyMethods[] = {
+		{ "__class_getitem__", PyCFunctionCast(FMethods::GetClassItem), METH_O|METH_CLASS, "__class_getitem__(cls, item: _ElemType) -- implemented for type hinting purpose only" },
 		{ "cast", PyCFunctionCast(&FMethods::Cast), METH_VARARGS | METH_KEYWORDS | METH_CLASS, "cast(cls, type: Type[_ElemType], obj: object) -> FixedArray[_ElemType] -- cast the given object to this Unreal fixed-array type" },
 		{ "__copy__", PyCFunctionCast(&FMethods::Copy), METH_NOARGS, "__copy__(self) -> FixedArray[_ElemType] -- copy this Unreal fixed-array" },
 		{ "copy", PyCFunctionCast(&FMethods::Copy), METH_NOARGS, "copy(self) -> FixedArray[_ElemType] -- copy this Unreal fixed-array" },
diff --git a/Engine/Plugins/Experimental/PythonScriptPlugin/Source/PythonScriptPlugin/Private/PyWrapperMap.cpp b/Engine/Plugins/Experimental/PythonScriptPlugin/Source/PythonScriptPlugin/Private/PyWrapperMap.cpp
--- a/Engine/Plugins/Experimental/PythonScriptPlugin/Source/PythonScriptPlugin/Private/PyWrapperMap.cpp
+++ b/Engine/Plugins/Experimental/PythonScriptPlugin/Source/PythonScriptPlugin/Private/PyWrapperMap.cpp
@@ -1340,6 +1340,13 @@ PyTypeObject InitializePyWrapperMapType()
 
 	struct FMethods
 	{
+		// This is for type hinting, to make the class appears as a generic (https://peps.python.org/pep-0560/) and support hinting like 'm: unreal.Map[int, str] = unreal.Map(int, str)'
+		static PyObject* GetClassItem(PyObject* Type, PyObject* Item)
+		{
+			Py_INCREF(Type);
+			return Type;
+		}
+
 		static PyObject* Cast(PyTypeObject* InType, PyObject* InArgs, PyObject* InKwds)
 		{
 			PyObject* PyKeyObj = nullptr;
@@ -1492,6 +1499,7 @@ PyTypeObject InitializePyWrapperMapType()
 
 	// NOTE: _KeyType = typing.TypeVar('_KeyType'), _ValueType = typing.TypeVar('_ValueType'), Type/Iterable/Tuple/Union/List/ItemsView/KeysView/ValuesView are defines in the Python typing module.
 	static PyMethodDef PyMethods[] = {
+		{ "__class_getitem__", PyCFunctionCast(FMethods::GetClassItem), METH_O|METH_CLASS, "__class_getitem__(cls, item) -- implemented for type hinting purpose only" },
 		{ "cast", PyCFunctionCast(&FMethods::Cast), METH_VARARGS | METH_KEYWORDS | METH_CLASS, "cast(cls, key: Type[_KeyType], value: Type[_ValueType], obj: object) -> Map[_KeyType, _ValueType] -- cast the given object to this Unreal map type" },
 		{ "__copy__", PyCFunctionCast(&FMethods::Copy), METH_NOARGS, "__copy__(self) -> Map[_KeyType, _ValueType] -- copy this Unreal map" },
 		{ "copy", PyCFunctionCast(&FMethods::Copy), METH_NOARGS, "copy(self) -> Map[_KeyType, _ValueType] -- copy this Unreal map" },
diff --git a/Engine/Plugins/Experimental/PythonScriptPlugin/Source/PythonScriptPlugin/Private/PyWrapperSet.cpp b/Engine/Plugins/Experimental/PythonScriptPlugin/Source/PythonScriptPlugin/Private/PyWrapperSet.cpp
--- a/Engine/Plugins/Experimental/PythonScriptPlugin/Source/PythonScriptPlugin/Private/PyWrapperSet.cpp
+++ b/Engine/Plugins/Experimental/PythonScriptPlugin/Source/PythonScriptPlugin/Private/PyWrapperSet.cpp
@@ -1228,6 +1228,13 @@ PyTypeObject InitializePyWrapperSetType()
 
 	struct FMethods
 	{
+		// This is for type hinting, to make the class appears as a generic (https://peps.python.org/pep-0560/) and support hinting like 's: unreal.Set[int] = unreal.Set(int)'
+		static PyObject* GetClassItem(PyObject* Type, PyObject* Item)
+		{
+			Py_INCREF(Type);
+			return Type;
+		}
+
 		static PyObject* Cast(PyTypeObject* InType, PyObject* InArgs, PyObject* InKwds)
 		{
 			PyObject* PyTypeObj = nullptr;
@@ -1476,6 +1483,7 @@ PyTypeObject InitializePyWrapperSetType()
 
 	// NOTE: _ElemType = typing.TypeVar('_ElemType') and Type/Iterable are defined by the Python typing module.
 	static PyMethodDef PyMethods[] = {
+		{ "__class_getitem__", PyCFunctionCast(FMethods::GetClassItem), METH_O|METH_CLASS, "__class_getitem__(cls, item: _ElemType) -- implemented for type hinting purpose only" },
 		{ "cast", PyCFunctionCast(&FMethods::Cast), METH_VARARGS | METH_KEYWORDS | METH_CLASS, "cast(cls, type: Type[_ElemType], obj: object) -> Set[_ElemType] -- cast the given object to this Unreal set type" },
 		{ "__copy__", PyCFunctionCast(&FMethods::Copy), METH_NOARGS, "__copy__(self) -> Set[_ElemType] -- copy this Unreal set" },
 		{ "copy", PyCFunctionCast(&FMethods::Copy), METH_NOARGS, "copy(self) -> Set[_ElemType] -- copy this Unreal set" },
