Detected and deprecated the construction of a subsystem from Python with a message telling user how to get the subsystem.
For example, this Python code now generates a deprecated warning: editor_subsystem = unreal.AssetEditorSubsystem()

diff --git a/Engine/Plugins/Experimental/PythonScriptPlugin/Source/PythonScriptPlugin/Private/PyUtil.cpp b/Engine/Plugins/Experimental/PythonScriptPlugin/Source/PythonScriptPlugin/Private/PyUtil.cpp
--- a/Engine/Plugins/Experimental/PythonScriptPlugin/Source/PythonScriptPlugin/Private/PyUtil.cpp
+++ b/Engine/Plugins/Experimental/PythonScriptPlugin/Source/PythonScriptPlugin/Private/PyUtil.cpp
@@ -661,6 +661,25 @@ UObject* NewObject(UClass* InObjClass, UObject* InObjectOuter, const FName InObj
 {
 	if (InObjClass)
 	{
+		if (InObjClass->IsChildOf(UEditorSubsystem::StaticClass()))
+		{
+			// Starting with UE 5.2, generates a warning/deprecation message.
+			SetPythonWarning(PyExc_DeprecationWarning, InErrorCtxt, *FString::Printf(TEXT("Editor subsystems creation is deprecated and will be removed in UE 5.3. Use 'unreal.get_editor_subsystem(unreal.%s)' to get an instance of the subsystem."), *PyGenUtil::GetClassPythonName(InObjClass)));
+
+			// For UE 5.3 or later, generate an hard error.
+			//SetPythonError(PyExc_Exception, InErrorCtxt, *FString::Printf(TEXT("Editor subsystems cannot be created. Use 'unreal.get_editor_subsystem(unreal.%s)' to get an instance of the subsystem."), *PyGenUtil::GetClassPythonName(InObjClass)));
+			//return nullptr;
+		}
+		else if (InObjClass->IsChildOf(UEngineSubsystem::StaticClass()))
+		{
+			// Starting with UE 5.2, generates a warning/deprecation message.
+			SetPythonWarning(PyExc_DeprecationWarning, InErrorCtxt, *FString::Printf(TEXT("Engine subsystems creation is deprecated and will be removed in UE 5.3. Use 'unreal.get_engine_subsystem(unreal.%s)' to get an instance of the subsystem."), *PyGenUtil::GetClassPythonName(InObjClass)));
+
+			// For UE 5.3 or later, generate an hard error.
+			//SetPythonError(PyExc_Exception, InErrorCtxt, *FString::Printf(TEXT("Engine subsystems cannot be created. Use 'unreal.get_engine_subsystem(unreal.%s)' to get an instance of the subsystem."), *PyGenUtil::GetClassPythonName(InObjClass)));
+			//return nullptr;
+		}
+
 		if (InObjClass == UPackage::StaticClass())
 		{
 			if (InObjectName.IsNone())
