Fix non-unity CIS build in PyUtil.cpp

diff --git a/Engine/Plugins/Experimental/PythonScriptPlugin/Source/PythonScriptPlugin/Private/PyUtil.cpp b/Engine/Plugins/Experimental/PythonScriptPlugin/Source/PythonScriptPlugin/Private/PyUtil.cpp
--- a/Engine/Plugins/Experimental/PythonScriptPlugin/Source/PythonScriptPlugin/Private/PyUtil.cpp
+++ b/Engine/Plugins/Experimental/PythonScriptPlugin/Source/PythonScriptPlugin/Private/PyUtil.cpp
@@ -31,8 +31,13 @@
 #include "UObject/TextProperty.h"
 #include "UObject/StructOnScope.h"
 #include "UObject/PropertyPortFlags.h"
+#include "Subsystems/EngineSubsystem.h"
 #include "Templates/Casts.h"
 
+#if WITH_EDITOR
+#include "EditorSubsystem.h"
+#endif // WITH_EDITOR
+
 DEFINE_LOG_CATEGORY(LogPython);
 
 #define LOCTEXT_NAMESPACE "Python"
@@ -661,24 +666,26 @@ UObject* NewObject(UClass* InObjClass, UObject* InObjectOuter, const FName InObj
 {
 	if (InObjClass)
 	{
-		if (InObjClass->IsChildOf(UEditorSubsystem::StaticClass()))
+		if (InObjClass->IsChildOf(UEngineSubsystem::StaticClass()))
 		{
 			// Starting with UE 5.2, generates a warning/deprecation message.
-			SetPythonWarning(PyExc_DeprecationWarning, InErrorCtxt, *FString::Printf(TEXT("Editor subsystems creation is deprecated and will be removed in UE 5.3. Use 'unreal.get_editor_subsystem(unreal.%s)' to get an instance of the subsystem."), *PyGenUtil::GetClassPythonName(InObjClass)));
+			SetPythonWarning(PyExc_DeprecationWarning, InErrorCtxt, *FString::Printf(TEXT("Engine subsystems creation is deprecated and will be removed in UE 5.3. Use 'unreal.get_engine_subsystem(unreal.%s)' to get an instance of the subsystem."), *PyGenUtil::GetClassPythonName(InObjClass)));
 
 			// For UE 5.3 or later, generate an hard error.
-			//SetPythonError(PyExc_Exception, InErrorCtxt, *FString::Printf(TEXT("Editor subsystems cannot be created. Use 'unreal.get_editor_subsystem(unreal.%s)' to get an instance of the subsystem."), *PyGenUtil::GetClassPythonName(InObjClass)));
+			//SetPythonError(PyExc_Exception, InErrorCtxt, *FString::Printf(TEXT("Engine subsystems cannot be created. Use 'unreal.get_engine_subsystem(unreal.%s)' to get an instance of the subsystem."), *PyGenUtil::GetClassPythonName(InObjClass)));
 			//return nullptr;
 		}
-		else if (InObjClass->IsChildOf(UEngineSubsystem::StaticClass()))
+#if WITH_EDITOR
+		else if (InObjClass->IsChildOf(UEditorSubsystem::StaticClass()))
 		{
 			// Starting with UE 5.2, generates a warning/deprecation message.
-			SetPythonWarning(PyExc_DeprecationWarning, InErrorCtxt, *FString::Printf(TEXT("Engine subsystems creation is deprecated and will be removed in UE 5.3. Use 'unreal.get_engine_subsystem(unreal.%s)' to get an instance of the subsystem."), *PyGenUtil::GetClassPythonName(InObjClass)));
+			SetPythonWarning(PyExc_DeprecationWarning, InErrorCtxt, *FString::Printf(TEXT("Editor subsystems creation is deprecated and will be removed in UE 5.3. Use 'unreal.get_editor_subsystem(unreal.%s)' to get an instance of the subsystem."), *PyGenUtil::GetClassPythonName(InObjClass)));
 
 			// For UE 5.3 or later, generate an hard error.
-			//SetPythonError(PyExc_Exception, InErrorCtxt, *FString::Printf(TEXT("Engine subsystems cannot be created. Use 'unreal.get_engine_subsystem(unreal.%s)' to get an instance of the subsystem."), *PyGenUtil::GetClassPythonName(InObjClass)));
+			//SetPythonError(PyExc_Exception, InErrorCtxt, *FString::Printf(TEXT("Editor subsystems cannot be created. Use 'unreal.get_editor_subsystem(unreal.%s)' to get an instance of the subsystem."), *PyGenUtil::GetClassPythonName(InObjClass)));
 			//return nullptr;
 		}
+#endif
 
 		if (InObjClass == UPackage::StaticClass())
 		{
