Added Python glue support for structs with internally mutable state

ScriptMethodMutable can be used to mark a ScriptMethod function taking a struct as needing to copy the input struct back out despite it being marked as const, due to the function changing internally mutable state within the struct. This is equivalent to passing the struct as UPARAM(ref), but allows the internal mutability of the struct to be preserved by the function signature.

diff --git a/Engine/Plugins/Experimental/PythonScriptPlugin/Source/PythonScriptPlugin/Private/PyGenUtil.cpp b/Engine/Plugins/Experimental/PythonScriptPlugin/Source/PythonScriptPlugin/Private/PyGenUtil.cpp
--- a/Engine/Plugins/Experimental/PythonScriptPlugin/Source/PythonScriptPlugin/Private/PyGenUtil.cpp
+++ b/Engine/Plugins/Experimental/PythonScriptPlugin/Source/PythonScriptPlugin/Private/PyGenUtil.cpp
@@ -57,6 +57,7 @@ const FName ScriptNameMetaDataKey = TEXT("ScriptName");
 const FName ScriptNoExportMetaDataKey = TEXT("ScriptNoExport");
 const FName ScriptMethodMetaDataKey = TEXT("ScriptMethod");
 const FName ScriptMethodSelfReturnMetaDataKey = TEXT("ScriptMethodSelfReturn");
+const FName ScriptMethodMutableMetaDataKey = TEXT("ScriptMethodMutable");
 const FName ScriptOperatorMetaDataKey = TEXT("ScriptOperator");
 const FName ScriptConstantMetaDataKey = TEXT("ScriptConstant");
 const FName ScriptConstantHostMetaDataKey = TEXT("ScriptConstantHost");
diff --git a/Engine/Plugins/Experimental/PythonScriptPlugin/Source/PythonScriptPlugin/Private/PyGenUtil.h b/Engine/Plugins/Experimental/PythonScriptPlugin/Source/PythonScriptPlugin/Private/PyGenUtil.h
--- a/Engine/Plugins/Experimental/PythonScriptPlugin/Source/PythonScriptPlugin/Private/PyGenUtil.h
+++ b/Engine/Plugins/Experimental/PythonScriptPlugin/Source/PythonScriptPlugin/Private/PyGenUtil.h
@@ -40,6 +40,7 @@ namespace PyGenUtil
 	extern const FName ScriptNoExportMetaDataKey;
 	extern const FName ScriptMethodMetaDataKey;
 	extern const FName ScriptMethodSelfReturnMetaDataKey;
+	extern const FName ScriptMethodMutableMetaDataKey;
 	extern const FName ScriptOperatorMetaDataKey;
 	extern const FName ScriptConstantMetaDataKey;
 	extern const FName ScriptConstantHostMetaDataKey;
diff --git a/Engine/Plugins/Experimental/PythonScriptPlugin/Source/PythonScriptPlugin/Private/PyTest.cpp b/Engine/Plugins/Experimental/PythonScriptPlugin/Source/PythonScriptPlugin/Private/PyTest.cpp
--- a/Engine/Plugins/Experimental/PythonScriptPlugin/Source/PythonScriptPlugin/Private/PyTest.cpp
+++ b/Engine/Plugins/Experimental/PythonScriptPlugin/Source/PythonScriptPlugin/Private/PyTest.cpp
@@ -52,6 +52,26 @@ FPyTestStruct UPyTestStructLibrary::AddStr(const FPyTestStruct& InStruct, const
 	return Result;
 }
 
+void UPyTestStructLibrary::SetBoolMutable(const FPyTestStruct& InStruct)
+{
+	InStruct.BoolMutable = true;
+}
+
+void UPyTestStructLibrary::ClearBoolMutable(const FPyTestStruct& InStruct)
+{
+	InStruct.BoolMutable = false;
+}
+
+void UPyTestStructLibrary::SetBoolMutableViaRef(FPyTestStruct& InStruct)
+{
+	InStruct.BoolMutable = true;
+}
+
+void UPyTestStructLibrary::ClearBoolMutableViaRef(FPyTestStruct& InStruct)
+{
+	InStruct.BoolMutable = false;
+}
+
 UPyTestInterface::UPyTestInterface(const FObjectInitializer& ObjectInitializer)
 	: Super(ObjectInitializer)
 {
diff --git a/Engine/Plugins/Experimental/PythonScriptPlugin/Source/PythonScriptPlugin/Private/PyTest.h b/Engine/Plugins/Experimental/PythonScriptPlugin/Source/PythonScriptPlugin/Private/PyTest.h
--- a/Engine/Plugins/Experimental/PythonScriptPlugin/Source/PythonScriptPlugin/Private/PyTest.h
+++ b/Engine/Plugins/Experimental/PythonScriptPlugin/Source/PythonScriptPlugin/Private/PyTest.h
@@ -88,6 +88,9 @@ public:
 
 	UPROPERTY(EditDefaultsOnly, Category = "Python|Internal")
 	bool BoolDefaultsOnly;
+
+	UPROPERTY(VisibleAnywhere, BlueprintReadOnly, Category = "Python|Internal")
+	mutable bool BoolMutable = false;
 };
 
 /**
@@ -116,14 +119,26 @@ class UPyTestStructLibrary : public UBlueprintFunctionLibrary
 	UFUNCTION(BlueprintPure, Category = "Python|Internal", meta=(ScriptConstant="ConstantValue", ScriptConstantHost="/Script/PythonScriptPlugin.PyTestStruct"))
 	static int32 GetConstantValue();
 
-	UFUNCTION(BlueprintPure, Category = "Python|Internal", meta=(ScriptMethod, ScriptMethodSelfReturn, ScriptOperator="+;+="))
+	UFUNCTION(BlueprintCallable, Category = "Python|Internal", meta=(ScriptMethod, ScriptMethodSelfReturn, ScriptOperator="+;+="))
 	static FPyTestStruct AddInt(const FPyTestStruct& InStruct, const int32 InValue);
 
-	UFUNCTION(BlueprintPure, Category = "Python|Internal", meta=(ScriptMethod, ScriptMethodSelfReturn, ScriptOperator="+;+="))
+	UFUNCTION(BlueprintCallable, Category = "Python|Internal", meta=(ScriptMethod, ScriptMethodSelfReturn, ScriptOperator="+;+="))
 	static FPyTestStruct AddFloat(const FPyTestStruct& InStruct, const float InValue);
 
-	UFUNCTION(BlueprintPure, Category = "Python|Internal", meta=(ScriptMethod, ScriptMethodSelfReturn, ScriptOperator="+;+="))
+	UFUNCTION(BlueprintCallable, Category = "Python|Internal", meta=(ScriptMethod, ScriptMethodSelfReturn, ScriptOperator="+;+="))
 	static FPyTestStruct AddStr(const FPyTestStruct& InStruct, const FString& InValue);
+
+	UFUNCTION(BlueprintCallable, Category = "Python|Internal", meta=(ScriptMethod, ScriptMethodMutable))
+	static void SetBoolMutable(const FPyTestStruct& InStruct);
+
+	UFUNCTION(BlueprintCallable, Category = "Python|Internal", meta=(ScriptMethod, ScriptMethodMutable))
+	static void ClearBoolMutable(const FPyTestStruct& InStruct);
+
+	UFUNCTION(BlueprintCallable, Category = "Python|Internal", meta=(ScriptMethod, ScriptMethodMutable))
+	static void SetBoolMutableViaRef(UPARAM(ref) FPyTestStruct& InStruct);
+
+	UFUNCTION(BlueprintCallable, Category = "Python|Internal", meta=(ScriptMethod, ScriptMethodMutable))
+	static void ClearBoolMutableViaRef(UPARAM(ref) FPyTestStruct& InStruct);
 };
 
 /**
diff --git a/Engine/Plugins/Experimental/PythonScriptPlugin/Source/PythonScriptPlugin/Private/PyWrapperTypeRegistry.cpp b/Engine/Plugins/Experimental/PythonScriptPlugin/Source/PythonScriptPlugin/Private/PyWrapperTypeRegistry.cpp
--- a/Engine/Plugins/Experimental/PythonScriptPlugin/Source/PythonScriptPlugin/Private/PyWrapperTypeRegistry.cpp
+++ b/Engine/Plugins/Experimental/PythonScriptPlugin/Source/PythonScriptPlugin/Private/PyWrapperTypeRegistry.cpp
@@ -1335,6 +1335,17 @@ PyTypeObject* FPyWrapperTypeRegistry::GenerateWrappedClassType(const UClass* InC
 			}
 		}
 
+		// The function may have been flagged as mutable, in which case we always consider it to need a 'self' return
+		if (!GeneratedWrappedDynamicMethod.SelfReturn.ParamProp && InFunc->HasMetaData(PyGenUtil::ScriptMethodMutableMetaDataKey))
+		{
+			if (!SelfParam.ParamProp->IsA<FStructProperty>())
+			{
+				REPORT_PYTHON_GENERATION_ISSUE(Error, TEXT("Function '%s.%s' is marked as 'ScriptMethodMutable' but the 'self' argument is not a struct."), *InFunc->GetOwnerClass()->GetName(), *InFunc->GetName());
+				return;
+			}
+			GeneratedWrappedDynamicMethod.SelfReturn = SelfParam;
+		}
+
 		// Set-up some data needed to build the tooltip correctly for the hoisted method
 		const bool bIsStaticOverride = false;
 		TSet<FName> ParamsToIgnore;
diff --git a/Engine/Source/Runtime/CoreUObject/Public/UObject/ObjectMacros.h b/Engine/Source/Runtime/CoreUObject/Public/UObject/ObjectMacros.h
--- a/Engine/Source/Runtime/CoreUObject/Public/UObject/ObjectMacros.h
+++ b/Engine/Source/Runtime/CoreUObject/Public/UObject/ObjectMacros.h
@@ -1556,9 +1556,12 @@ namespace UM
 		/// The value is optional, and may specify a name override for the method. May include deprecated names as additional semi-colon separated entries.
 		ScriptMethod,
 
-		/// [FunctionMetadata] Used with ScriptMethod to denote that the return value of the function should overwrite the value of the instance that made the call (structs only, equivalent to using UPARAM(self) on the struct argument).
+		/// [FunctionMetadata] Used with ScriptMethod to denote that the return value of the function should overwrite the value of the instance that made the call (structs only, equivalent to using UPARAM(ref) on the struct argument).
 		ScriptMethodSelfReturn,
 
+		/// [FunctionMetadata] Used with ScriptMethod to denote that the first argument of the function should be considered mutable even if passed as a const reference (structs only, equivalent to using UPARAM(ref) on the struct argument).
+		ScriptMethodMutable,
+
 		/// [FunctionMetadata] Flags a static function taking a struct as its first argument so that it "hoists" the function to be an operator of the struct when exporting it to a scripting language.
 		/// The value describes the kind of operator using C++ operator syntax (see below), and may contain multiple semi-colon separated values.
 		/// The signature of the function depends on the operator type, and additional parameters may be passed as long as they're defaulted and the basic signature requirements are met.
