Updated the PythonScriptPlugin to run the startup scripts after "EndOfEngineInit" instead of the first editor tick.  This fixes an issue where the user would see the editor hang after it appeared after the loadscreen. The loadscreen now says that it is initializing the scripts at the end.

diff --git a/Engine/Plugins/Experimental/PythonScriptPlugin/Source/PythonScriptPlugin/Private/PythonScriptPlugin.cpp b/Engine/Plugins/Experimental/PythonScriptPlugin/Source/PythonScriptPlugin/Private/PythonScriptPlugin.cpp
--- a/Engine/Plugins/Experimental/PythonScriptPlugin/Source/PythonScriptPlugin/Private/PythonScriptPlugin.cpp
+++ b/Engine/Plugins/Experimental/PythonScriptPlugin/Source/PythonScriptPlugin/Private/PythonScriptPlugin.cpp
@@ -23,9 +23,11 @@
 #include "Misc/CoreDelegates.h"
 #include "Misc/Char.h"
 #include "Misc/CString.h"
+#include "Misc/DelayedAutoRegister.h"
 #include "Misc/Paths.h"
 #include "Misc/FileHelper.h"
 #include "Misc/PackageName.h"
+#include "Misc/ScopedSlowTask.h"
 #include "HAL/FileManager.h"
 #include "HAL/LowLevelMemTracker.h"
 #include "HAL/PlatformMisc.h"
@@ -474,13 +476,23 @@ private:
 
 #endif	// WITH_PYTHON
 
+#if WITH_PYTHON
+static FDelayedAutoRegisterHelper GPythonScriptAutoRegistrySingletonHelper(EDelayedRegisterRunPhase::EndOfEngineInit, []() -> void
+	{
+		if (FPythonScriptPlugin::Get()->IsPythonAvailable())
+		{
+			FPythonScriptPlugin::Get()->RunStartupScripts();		
+		}
+	});
+#endif
+
 FPythonScriptPlugin::FPythonScriptPlugin()
 #if WITH_PYTHON
 	: CmdExec(this)
 	, CmdREPLExec(this)
 	, CmdMenu(nullptr)
 	, bInitialized(false)
-	, bHasTicked(false)
+	, bRanStartupScripts(false)
 #endif	// WITH_PYTHON
 {
 }
@@ -1110,13 +1122,13 @@ void FPythonScriptPlugin::ShutdownPython()
 	Py_Finalize();
 
 	bInitialized = false;
-	bHasTicked = false;
+	bRanStartupScripts = false;
 }
 
 void FPythonScriptPlugin::RequestStubCodeGeneration()
 {
-	// Ignore requests made before the fist Tick
-	if (!bHasTicked)
+	// Ignore requests made before the startup scripts have ran
+	if (!bRanStartupScripts)
 	{
 		return;
 	}
@@ -1160,51 +1172,67 @@ void FPythonScriptPlugin::GenerateStubCode()
 	}
 }
 
-void FPythonScriptPlugin::Tick(const float InDeltaTime)
+void FPythonScriptPlugin::RunStartupScripts()
 {
-	// If this is our first Tick, handle any post-init logic that should happen once the engine is fully initialized
-	if (!bHasTicked)
+	if (bRanStartupScripts)
 	{
-		bHasTicked = true;
+		return;
+	}
 
-		// Run start-up scripts now
-		TArray<FString> PySysPaths;
-		{
-			FPyScopedGIL GIL;
-			PySysPaths = PyUtil::GetSystemPaths();
-		}
-		for (const FString& PySysPath : PySysPaths)
+	bRanStartupScripts = true;
+
+	// Run start-up scripts now
+	TArray<FString> PySysPaths;
+	{
+		FPyScopedGIL GIL;
+		PySysPaths = PyUtil::GetSystemPaths();
+	}
+
+	FScopedSlowTask Progress(PySysPaths.Num() + GetDefault<UPythonScriptPluginSettings>()->StartupScripts.Num(), LOCTEXT("PythonScriptPluginInitScripts", "Running Python start-up scripts..."), true);
+	for (const FString& PySysPath : PySysPaths)
+	{
+		const FString PotentialFilePath = PySysPath / TEXT("init_unreal.py");
+		if (FPaths::FileExists(PotentialFilePath))
 		{
-			const FString PotentialFilePath = PySysPath / TEXT("init_unreal.py");
-			if (FPaths::FileExists(PotentialFilePath))
-			{
-				// Execute these files in the "public" scope, as if their contents had been run directly in the console
-				// This allows them to be used to set-up an editor environment for the console
-				FPythonCommandEx InitUnrealPythonCommand;
-				InitUnrealPythonCommand.FileExecutionScope = EPythonFileExecutionScope::Public;
-				RunFile(*PotentialFilePath, *InitUnrealPythonCommand.Command, InitUnrealPythonCommand);
-			}
+			Progress.EnterProgressFrame(1.0f, FText::Format(LOCTEXT("PythonScriptPluginInitScripts_Running", "Running {0}..."), FText::FromString(PotentialFilePath)));
+
+			// Execute these files in the "public" scope, as if their contents had been run directly in the console
+			// This allows them to be used to set-up an editor environment for the console
+			FPythonCommandEx InitUnrealPythonCommand;
+			InitUnrealPythonCommand.FileExecutionScope = EPythonFileExecutionScope::Public;
+			RunFile(*PotentialFilePath, *InitUnrealPythonCommand.Command, InitUnrealPythonCommand);
 		}
-		for (const FString& StartupScript : GetDefault<UPythonScriptPluginSettings>()->StartupScripts)
+		else
 		{
-			ExecPythonCommand(*StartupScript);
+			Progress.EnterProgressFrame();
 		}
+	}
+
+	for (const FString& StartupScript : GetDefault<UPythonScriptPluginSettings>()->StartupScripts)
+	{
+		Progress.EnterProgressFrame(1.0f, FText::Format(LOCTEXT("PythonScriptPluginInitScripts_Running", "Running {0}..."), FText::FromString(StartupScript)));
+		ExecPythonCommand(*StartupScript);
+	}
 
-		// Notify any external listeners
-		OnPythonInitializedDelegate.Broadcast();
+	// Notify any external listeners
+	OnPythonInitializedDelegate.Broadcast();
 
 #if WITH_EDITOR
-		// Activate the Content Browser integration (now that editor subsystems are available)
-		if (PythonFileDataSource)
-		{
-			UContentBrowserDataSubsystem* ContentBrowserData = IContentBrowserDataModule::Get().GetSubsystem();
-			ContentBrowserData->ActivateDataSource("PythonData");
-		}
+	// Activate the Content Browser integration (now that editor subsystems are available)
+	if (PythonFileDataSource)
+	{
+		UContentBrowserDataSubsystem* ContentBrowserData = IContentBrowserDataModule::Get().GetSubsystem();
+		ContentBrowserData->ActivateDataSource("PythonData");
+	}
 
-		// Register to generate stub code after a short delay
-		RequestStubCodeGeneration();
+	// Register to generate stub code after a short delay
+	RequestStubCodeGeneration();
 #endif	// WITH_EDITOR
-	}
+}
+
+void FPythonScriptPlugin::Tick(const float InDeltaTime)
+{
+	RunStartupScripts();
 
 	RemoteExecution->Tick(InDeltaTime);
 
diff --git a/Engine/Plugins/Experimental/PythonScriptPlugin/Source/PythonScriptPlugin/Private/PythonScriptPlugin.h b/Engine/Plugins/Experimental/PythonScriptPlugin/Source/PythonScriptPlugin/Private/PythonScriptPlugin.h
--- a/Engine/Plugins/Experimental/PythonScriptPlugin/Source/PythonScriptPlugin/Private/PythonScriptPlugin.h
+++ b/Engine/Plugins/Experimental/PythonScriptPlugin/Source/PythonScriptPlugin/Private/PythonScriptPlugin.h
@@ -123,6 +123,7 @@ public:
 	virtual void PostChange(const UUserDefinedEnum* Enum, FEnumEditorUtils::EEnumEditorChangeInfo Info) override;
 
 #if WITH_PYTHON
+	void RunStartupScripts();
 
 #if WITH_EDITOR	
 	void OnPostEngineInit();
@@ -213,7 +214,7 @@ private:
 	FPyObjectPtr PyUnrealModule;
 	PyThreadState* PyMainThreadState = nullptr;
 	bool bInitialized;
-	bool bHasTicked;
+	bool bRanStartupScripts;
 #endif	// WITH_PYTHON
 
 	FSimpleMulticastDelegate OnPythonInitializedDelegate;
