PR #9726: Replaced deprecated functions (Contributed by [REDACTED_PERSON])

diff --git a/Engine/Plugins/Experimental/PythonScriptPlugin/Source/PythonScriptPlugin/Private/PythonScriptPlugin.cpp b/Engine/Plugins/Experimental/PythonScriptPlugin/Source/PythonScriptPlugin/Private/PythonScriptPlugin.cpp
--- a/Engine/Plugins/Experimental/PythonScriptPlugin/Source/PythonScriptPlugin/Private/PythonScriptPlugin.cpp
+++ b/Engine/Plugins/Experimental/PythonScriptPlugin/Source/PythonScriptPlugin/Private/PythonScriptPlugin.cpp
@@ -1318,23 +1318,7 @@ PyObject* FPythonScriptPlugin::EvalString(const TCHAR* InStr, const TCHAR* InCon
 {
 	TRACE_CPUPROFILER_EVENT_SCOPE(FPythonScriptPlugin::EvalString)
 
-	PyCompilerFlags *PyCompFlags = nullptr;
-
-	PyArena* PyArena = PyArena_New();
-	if (!PyArena)
-	{
-		return nullptr;
-	}
-
-	_mod* PyModule = PyParser_ASTFromString(TCHAR_TO_UTF8(InStr), TCHAR_TO_UTF8(InContext), InMode, PyCompFlags, PyArena);
-	if (!PyModule)
-	{
-		PyArena_Free(PyArena);
-		return nullptr;
-	}
-
-	typedef TPyPtr<PyCodeObject> PyCodeObjectPtr;
-	PyCodeObjectPtr PyCodeObj = PyCodeObjectPtr::StealReference(PyAST_Compile(PyModule, TCHAR_TO_UTF8(InContext), PyCompFlags, PyArena));
+	FPyObjectPtr PyCodeObj = FPyObjectPtr::StealReference(Py_CompileString(TCHAR_TO_UTF8(InStr), TCHAR_TO_UTF8(InContext), InMode));
 	if (!PyCodeObj)
 	{
 		return nullptr;
