Removed starting the python scripts during the load screen because some scripts had issues. It will now show a progress dialog if it is taking too long to process.

diff --git a/Engine/Plugins/Experimental/PythonScriptPlugin/Source/PythonScriptPlugin/Private/PythonScriptPlugin.cpp b/Engine/Plugins/Experimental/PythonScriptPlugin/Source/PythonScriptPlugin/Private/PythonScriptPlugin.cpp
--- a/Engine/Plugins/Experimental/PythonScriptPlugin/Source/PythonScriptPlugin/Private/PythonScriptPlugin.cpp
+++ b/Engine/Plugins/Experimental/PythonScriptPlugin/Source/PythonScriptPlugin/Private/PythonScriptPlugin.cpp
@@ -476,16 +476,6 @@ private:
 
 #endif	// WITH_PYTHON
 
-#if WITH_PYTHON
-static FDelayedAutoRegisterHelper GPythonScriptAutoRegistrySingletonHelper(EDelayedRegisterRunPhase::EndOfEngineInit, []() -> void
-	{
-		if (FPythonScriptPlugin::Get()->IsPythonAvailable())
-		{
-			FPythonScriptPlugin::Get()->RunStartupScripts();		
-		}
-	});
-#endif
-
 FPythonScriptPlugin::FPythonScriptPlugin()
 #if WITH_PYTHON
 	: CmdExec(this)
@@ -1189,6 +1179,7 @@ void FPythonScriptPlugin::RunStartupScripts()
 	}
 
 	FScopedSlowTask Progress(PySysPaths.Num() + GetDefault<UPythonScriptPluginSettings>()->StartupScripts.Num(), LOCTEXT("PythonScriptPluginInitScripts", "Running Python start-up scripts..."), true);
+	Progress.MakeDialogDelayed(0.1f);
 	for (const FString& PySysPath : PySysPaths)
 	{
 		const FString PotentialFilePath = PySysPath / TEXT("init_unreal.py");
