Added extended error information for a SyntaxError in Python

This now reports the error line and the offset within it that the error occurred on, which makes it consitent with the error reporting from PyErr_Print()

diff --git a/Engine/Plugins/Experimental/PythonScriptPlugin/Source/PythonScriptPlugin/Private/PyUtil.cpp b/Engine/Plugins/Experimental/PythonScriptPlugin/Source/PythonScriptPlugin/Private/PyUtil.cpp
--- a/Engine/Plugins/Experimental/PythonScriptPlugin/Source/PythonScriptPlugin/Private/PyUtil.cpp
+++ b/Engine/Plugins/Experimental/PythonScriptPlugin/Source/PythonScriptPlugin/Private/PyUtil.cpp
@@ -1428,6 +1428,30 @@ bool FetchPythonError(FString& OutError)
 		{
 			FPyObjectPtr PyExceptionTypeName = FPyObjectPtr::StealReference(PyObject_GetAttrString(PyExceptionType, "__name__"));
 			OutError = FString::Printf(TEXT("%s: %s"), PyExceptionTypeName ? *PyObjectToUEString(PyExceptionTypeName) : *PyObjectToUEString(PyExceptionType), *PyObjectToUEString(PyExceptionValue));
+
+			// Syntax errors require special handling to get the extended error information out of them (their str() only returns the error message and not the context info)
+			if (PyExceptionType == PyExc_SyntaxError)
+			{
+				FPyObjectPtr PySyntaxErrorText = FPyObjectPtr::StealReference(PyObject_GetAttrString(PyExceptionValue, "text"));
+				FPyObjectPtr PySyntaxErrorOffset = FPyObjectPtr::StealReference(PyObject_GetAttrString(PyExceptionValue, "offset"));
+				if (PySyntaxErrorText && PySyntaxErrorOffset)
+				{
+					OutError += TEXT("\n  ");
+					OutError += PyObjectToUEString(PySyntaxErrorText);
+
+					int32 OffsetInt = 0;
+					PyConversion::Nativize(PySyntaxErrorOffset, OffsetInt, PyConversion::ESetErrorState::No);
+					if (OffsetInt > 0)
+					{
+						OutError += TEXT("\n  ");
+						for (int32 OffsetIndex = 1; OffsetIndex < OffsetInt; ++OffsetIndex)
+						{
+							OutError += TEXT(" ");
+						}
+						OutError += TEXT("^");
+					}
+				}
+			}
 		}
 		else
 		{
