Also handle extended information for SyntaxError-based exceptions (like IndentationError)

#rnx

diff --git a/Engine/Plugins/Experimental/PythonScriptPlugin/Source/PythonScriptPlugin/Private/PyUtil.cpp b/Engine/Plugins/Experimental/PythonScriptPlugin/Source/PythonScriptPlugin/Private/PyUtil.cpp
--- a/Engine/Plugins/Experimental/PythonScriptPlugin/Source/PythonScriptPlugin/Private/PyUtil.cpp
+++ b/Engine/Plugins/Experimental/PythonScriptPlugin/Source/PythonScriptPlugin/Private/PyUtil.cpp
@@ -1430,7 +1430,7 @@ bool FetchPythonError(FString& OutError)
 			OutError = FString::Printf(TEXT("%s: %s"), PyExceptionTypeName ? *PyObjectToUEString(PyExceptionTypeName) : *PyObjectToUEString(PyExceptionType), *PyObjectToUEString(PyExceptionValue));
 
 			// Syntax errors require special handling to get the extended error information out of them (their str() only returns the error message and not the context info)
-			if (PyExceptionType == PyExc_SyntaxError)
+			if (PyObject_IsInstance(PyExceptionValue, PyExc_SyntaxError))
 			{
 				FPyObjectPtr PySyntaxErrorText = FPyObjectPtr::StealReference(PyObject_GetAttrString(PyExceptionValue, "text"));
 				FPyObjectPtr PySyntaxErrorOffset = FPyObjectPtr::StealReference(PyObject_GetAttrString(PyExceptionValue, "offset"));
