Added ScriptDefaultMake and ScriptDefaultBreak to instruct script layers to always use default make/break functions

BlueprintInternalUseOnly and BlueprintInternalUseOnlyHierarchical is now respected for structs in python glue code

diff --git a/Engine/Plugins/Experimental/PythonScriptPlugin/Source/PythonScriptPlugin/Private/PyGenUtil.cpp b/Engine/Plugins/Experimental/PythonScriptPlugin/Source/PythonScriptPlugin/Private/PyGenUtil.cpp
--- a/Engine/Plugins/Experimental/PythonScriptPlugin/Source/PythonScriptPlugin/Private/PyGenUtil.cpp
+++ b/Engine/Plugins/Experimental/PythonScriptPlugin/Source/PythonScriptPlugin/Private/PyGenUtil.cpp
@@ -61,12 +61,15 @@ const FName ScriptMethodMutableMetaDataKey = TEXT("ScriptMethodMutable");
 const FName ScriptOperatorMetaDataKey = TEXT("ScriptOperator");
 const FName ScriptConstantMetaDataKey = TEXT("ScriptConstant");
 const FName ScriptConstantHostMetaDataKey = TEXT("ScriptConstantHost");
+const FName ScriptDefaultMakeMetaDataKey = TEXT("ScriptDefaultMake");
+const FName ScriptDefaultBreakMetaDataKey = TEXT("ScriptDefaultBreak");
 const FName BlueprintTypeMetaDataKey = TEXT("BlueprintType");
 const FName NotBlueprintTypeMetaDataKey = TEXT("NotBlueprintType");
 const FName BlueprintSpawnableComponentMetaDataKey = TEXT("BlueprintSpawnableComponent");
 const FName BlueprintGetterMetaDataKey = TEXT("BlueprintGetter");
 const FName BlueprintSetterMetaDataKey = TEXT("BlueprintSetter");
 const FName BlueprintInternalUseOnlyMetaDataKey = TEXT("BlueprintInternalUseOnly");
+const FName BlueprintInternalUseOnlyHierarchicalMetaDataKey = TEXT("BlueprintInternalUseOnlyHierarchical");
 const FName CustomThunkMetaDataKey = TEXT("CustomThunk");
 const FName DeprecatedPropertyMetaDataKey = TEXT("DeprecatedProperty");
 const FName DeprecatedFunctionMetaDataKey = TEXT("DeprecatedFunction");
@@ -1646,6 +1649,19 @@ bool ShouldExportClass(const UClass* InClass)
 
 bool ShouldExportStruct(const UScriptStruct* InStruct)
 {
+	if (InStruct->HasMetaData(BlueprintInternalUseOnlyMetaDataKey))
+	{
+		return false;
+	}
+
+	for (const UScriptStruct* ParentStruct = InStruct; ParentStruct; ParentStruct = Cast<UScriptStruct>(ParentStruct->GetSuperStruct()))
+	{
+		if (ParentStruct->HasMetaData(BlueprintInternalUseOnlyHierarchicalMetaDataKey))
+		{
+			return false;
+		}
+	}
+
 	return IsScriptExposedStruct(InStruct) || HasScriptExposedFields(InStruct);
 }
 
diff --git a/Engine/Plugins/Experimental/PythonScriptPlugin/Source/PythonScriptPlugin/Private/PyGenUtil.h b/Engine/Plugins/Experimental/PythonScriptPlugin/Source/PythonScriptPlugin/Private/PyGenUtil.h
--- a/Engine/Plugins/Experimental/PythonScriptPlugin/Source/PythonScriptPlugin/Private/PyGenUtil.h
+++ b/Engine/Plugins/Experimental/PythonScriptPlugin/Source/PythonScriptPlugin/Private/PyGenUtil.h
@@ -44,12 +44,15 @@ namespace PyGenUtil
 	extern const FName ScriptOperatorMetaDataKey;
 	extern const FName ScriptConstantMetaDataKey;
 	extern const FName ScriptConstantHostMetaDataKey;
+	extern const FName ScriptDefaultMakeMetaDataKey;
+	extern const FName ScriptDefaultBreakMetaDataKey;
 	extern const FName BlueprintTypeMetaDataKey;
 	extern const FName NotBlueprintTypeMetaDataKey;
 	extern const FName BlueprintSpawnableComponentMetaDataKey;
 	extern const FName BlueprintGetterMetaDataKey;
 	extern const FName BlueprintSetterMetaDataKey;
 	extern const FName BlueprintInternalUseOnlyMetaDataKey;
+	extern const FName BlueprintInternalUseOnlyHierarchicalMetaDataKey;
 	extern const FName CustomThunkMetaDataKey;
 	extern const FName HasNativeMakeMetaDataKey;
 	extern const FName HasNativeBreakMetaDataKey;
diff --git a/Engine/Plugins/Experimental/PythonScriptPlugin/Source/PythonScriptPlugin/Private/PyWrapperTypeRegistry.cpp b/Engine/Plugins/Experimental/PythonScriptPlugin/Source/PythonScriptPlugin/Private/PyWrapperTypeRegistry.cpp
--- a/Engine/Plugins/Experimental/PythonScriptPlugin/Source/PythonScriptPlugin/Private/PyWrapperTypeRegistry.cpp
+++ b/Engine/Plugins/Experimental/PythonScriptPlugin/Source/PythonScriptPlugin/Private/PyWrapperTypeRegistry.cpp
@@ -47,25 +47,29 @@ namespace UE::Python
 /**
  * Finds the UFunction corresponding to the name specified by 'HasNativeMake' or 'HasNativeBreak' meta data key.
  * @param The structure to inspect for the 'HasNativeMake' or 'HasNativeBreak' meta data keys.
- * @param InMetaDataKey The meta data key name. Can only be 'HasNativeMake' or 'HasNativeBreak'.
+ * @param InNativeMetaDataKey The native meta data key name. Can only be 'HasNativeMake' or 'HasNativeBreak'.
+ * @param InScriptDefaultMetaDataKey The script default meta data key name. Can only be 'ScriptDefaultMake' or 'ScriptDefaultBreak'.
  * @param NotFoundFn Function invoked if the structure specifies as Make or Break function, but the function couldn't be found.
  * @return The function, if the struct has the meta key and if the function was found. Null otherwise.
  */
 template<typename NotFoundFuncT>
-UFunction* FindMakeBreakFunction(const UScriptStruct* InStruct, const FName& InMetaDataKey, const NotFoundFuncT& NotFoundFn)
+UFunction* FindMakeBreakFunction(const UScriptStruct* InStruct, const FName& InNativeMetaDataKey, const FName& InScriptDefaultMetaDataKey, const NotFoundFuncT& NotFoundFn)
 {
-	check(InMetaDataKey == PyGenUtil::HasNativeMakeMetaDataKey || InMetaDataKey == PyGenUtil::HasNativeBreakMetaDataKey);
+	check(InNativeMetaDataKey == PyGenUtil::HasNativeMakeMetaDataKey || InNativeMetaDataKey == PyGenUtil::HasNativeBreakMetaDataKey);
+	check(InScriptDefaultMetaDataKey == PyGenUtil::ScriptDefaultMakeMetaDataKey || InScriptDefaultMetaDataKey == PyGenUtil::ScriptDefaultBreakMetaDataKey);
 
 	UFunction* MakeBreakFunc = nullptr;
-
-	const FString MakeBreakFunctionName = InStruct->GetMetaData(InMetaDataKey);
-	if (!MakeBreakFunctionName.IsEmpty())
+	if (!InStruct->HasMetaData(InScriptDefaultMetaDataKey))
 	{
-		// Find the function.
-		MakeBreakFunc = FindObject<UFunction>(/*Outer*/nullptr, *MakeBreakFunctionName, /*ExactClass*/true);
-		if (!MakeBreakFunc)
+		const FString MakeBreakFunctionName = InStruct->GetMetaData(InNativeMetaDataKey);
+		if (!MakeBreakFunctionName.IsEmpty())
 		{
-			NotFoundFn(MakeBreakFunctionName);
+			// Find the function.
+			MakeBreakFunc = FindObject<UFunction>(/*Outer*/nullptr, *MakeBreakFunctionName, /*ExactClass*/true);
+			if (!MakeBreakFunc)
+			{
+				NotFoundFn(MakeBreakFunctionName);
+			}
 		}
 	}
 	return MakeBreakFunc;
@@ -2016,14 +2020,14 @@ PyTypeObject* FPyWrapperTypeRegistry::GenerateWrappedStructType(const UScriptStr
 	StructMetaData->PythonDeprecatedProperties = MoveTemp(PythonDeprecatedProperties);
 
 	// If the struct has the 'HasNativeMake' meta and the function is found, set it, otherwise, postpone maybe its going to be loaded later.
-	if (UFunction* MakeFunc = UE::Python::FindMakeBreakFunction(InStruct, PyGenUtil::HasNativeMakeMetaDataKey,
+	if (UFunction* MakeFunc = UE::Python::FindMakeBreakFunction(InStruct, PyGenUtil::HasNativeMakeMetaDataKey, PyGenUtil::ScriptDefaultMakeMetaDataKey, 
 		[this, &StructMetaData](const FString& MakeFuncName){ UnresolvedMakeFuncs.Emplace(MakeFuncName, StructMetaData); }))
 	{
 		UE::Python::SetMakeFunction(*StructMetaData, MakeFunc);
 	}
 
 	// If the struct has the 'HasNativeBreak' meta and the function is found, set it, otherwise, postpone maybe its going to be loaded later.
-	if (UFunction* BreakFunc = UE::Python::FindMakeBreakFunction(InStruct, PyGenUtil::HasNativeBreakMetaDataKey,
+	if (UFunction* BreakFunc = UE::Python::FindMakeBreakFunction(InStruct, PyGenUtil::HasNativeBreakMetaDataKey, PyGenUtil::ScriptDefaultBreakMetaDataKey, 
 		[this, &StructMetaData](const FString& BreakFuncName) { UnresolvedBreakFuncs.Emplace(BreakFuncName, StructMetaData); }))
 	{
 		UE::Python::SetBreakFunction(*StructMetaData, BreakFunc);
diff --git a/Engine/Source/Runtime/SlateCore/Public/Types/SlateVector2.h b/Engine/Source/Runtime/SlateCore/Public/Types/SlateVector2.h
--- a/Engine/Source/Runtime/SlateCore/Public/Types/SlateVector2.h
+++ b/Engine/Source/Runtime/SlateCore/Public/Types/SlateVector2.h
@@ -113,7 +113,7 @@ namespace UE::Slate
  *
  * NOTE: DO NOT USE THIS TYPE DIRECTLY IN CLIENT CODE - ALL USAGES SHOULD USE FVector2f or float
  */
-USTRUCT(BlueprintInternalUseOnly, DisplayName="Vector2D", meta=(HiddenByDefault, ShortTooltip="Vector2D (single-precision)", ToolTip="Vector2D (single-precision)", HasNativeMake="/Script/Engine.KismetMathLibrary.MakeVector2D", HasNativeBreak="/Script/Engine.KismetMathLibrary.BreakVector2D"))
+USTRUCT(BlueprintInternalUseOnly, DisplayName="Vector2D", meta=(HiddenByDefault, ShortTooltip="Vector2D (single-precision)", ToolTip="Vector2D (single-precision)", ScriptDefaultMake, ScriptDefaultBreak, HasNativeMake="/Script/Engine.KismetMathLibrary.MakeVector2D", HasNativeBreak="/Script/Engine.KismetMathLibrary.BreakVector2D"))
 struct FDeprecateSlateVector2D : public FVector2f
 {
 	GENERATED_BODY()
