Python: Disabling BlueprintInternalUseOnly/Hierarchical for structs for now, as some existing scripts depend on these internal types

#rnx

diff --git a/Engine/Plugins/Experimental/PythonScriptPlugin/Source/PythonScriptPlugin/Private/PyGenUtil.cpp b/Engine/Plugins/Experimental/PythonScriptPlugin/Source/PythonScriptPlugin/Private/PyGenUtil.cpp
--- a/Engine/Plugins/Experimental/PythonScriptPlugin/Source/PythonScriptPlugin/Private/PyGenUtil.cpp
+++ b/Engine/Plugins/Experimental/PythonScriptPlugin/Source/PythonScriptPlugin/Private/PyGenUtil.cpp
@@ -1649,6 +1649,8 @@ bool ShouldExportClass(const UClass* InClass)
 
 bool ShouldExportStruct(const UScriptStruct* InStruct)
 {
+	// Disabling BlueprintInternalUseOnly/Hierarchical for structs for now, as some existing scripts depend on these internal types
+	/*
 	if (InStruct->HasMetaData(BlueprintInternalUseOnlyMetaDataKey))
 	{
 		return false;
@@ -1661,6 +1663,7 @@ bool ShouldExportStruct(const UScriptStruct* InStruct)
 			return false;
 		}
 	}
+	*/
 
 	return IsScriptExposedStruct(InStruct) || HasScriptExposedFields(InStruct);
 }
