Exposed GetInterpreterExecutablePath via IPythonScriptPlugin

diff --git a/Engine/Plugins/Experimental/PythonScriptPlugin/Source/PythonScriptPlugin/Private/PythonScriptPlugin.cpp b/Engine/Plugins/Experimental/PythonScriptPlugin/Source/PythonScriptPlugin/Private/PythonScriptPlugin.cpp
--- a/Engine/Plugins/Experimental/PythonScriptPlugin/Source/PythonScriptPlugin/Private/PythonScriptPlugin.cpp
+++ b/Engine/Plugins/Experimental/PythonScriptPlugin/Source/PythonScriptPlugin/Private/PythonScriptPlugin.cpp
@@ -607,6 +607,15 @@ bool FPythonScriptPlugin::ExecPythonCommandEx(FPythonCommandEx& InOutPythonComma
 #endif	// WITH_PYTHON
 }
 
+FString FPythonScriptPlugin::GetInterpreterExecutablePath() const
+{
+#if WITH_PYTHON
+	return PyUtil::GetInterpreterExecutablePath();
+#else	// WITH_PYTHON
+	return FString();
+#endif	// WITH_PYTHON
+}
+
 FSimpleMulticastDelegate& FPythonScriptPlugin::OnPythonInitialized()
 {
 	return OnPythonInitializedDelegate;
diff --git a/Engine/Plugins/Experimental/PythonScriptPlugin/Source/PythonScriptPlugin/Private/PythonScriptPlugin.h b/Engine/Plugins/Experimental/PythonScriptPlugin/Source/PythonScriptPlugin/Private/PythonScriptPlugin.h
--- a/Engine/Plugins/Experimental/PythonScriptPlugin/Source/PythonScriptPlugin/Private/PythonScriptPlugin.h
+++ b/Engine/Plugins/Experimental/PythonScriptPlugin/Source/PythonScriptPlugin/Private/PythonScriptPlugin.h
@@ -108,6 +108,7 @@ public:
 	virtual bool IsPythonAvailable() const override;
 	virtual bool ExecPythonCommand(const TCHAR* InPythonCommand) override;
 	virtual bool ExecPythonCommandEx(FPythonCommandEx& InOutPythonCommand) override;
+	virtual FString GetInterpreterExecutablePath() const override;
 	virtual FSimpleMulticastDelegate& OnPythonInitialized() override;
 	virtual FSimpleMulticastDelegate& OnPythonShutdown() override;
 
diff --git a/Engine/Plugins/Experimental/PythonScriptPlugin/Source/PythonScriptPlugin/Public/IPythonScriptPlugin.h b/Engine/Plugins/Experimental/PythonScriptPlugin/Source/PythonScriptPlugin/Public/IPythonScriptPlugin.h
--- a/Engine/Plugins/Experimental/PythonScriptPlugin/Source/PythonScriptPlugin/Public/IPythonScriptPlugin.h
+++ b/Engine/Plugins/Experimental/PythonScriptPlugin/Source/PythonScriptPlugin/Public/IPythonScriptPlugin.h
@@ -35,6 +35,11 @@ public:
 	 */
 	virtual bool ExecPythonCommandEx(FPythonCommandEx& InOutPythonCommand) = 0;
 	
+	/**
+	 * Get the path to the Python interpreter executable of the Python SDK this plugin was compiled against.
+	 */
+	virtual FString GetInterpreterExecutablePath() const = 0;
+
 	/**
 	 * Delegate called after Python has been initialized.
 	 */
