Log the execution duration of each Python start-up script

#rnx

diff --git a/Engine/Plugins/Experimental/PythonScriptPlugin/Source/PythonScriptPlugin/Private/PythonScriptPlugin.cpp b/Engine/Plugins/Experimental/PythonScriptPlugin/Source/PythonScriptPlugin/Private/PythonScriptPlugin.cpp
--- a/Engine/Plugins/Experimental/PythonScriptPlugin/Source/PythonScriptPlugin/Private/PythonScriptPlugin.cpp
+++ b/Engine/Plugins/Experimental/PythonScriptPlugin/Source/PythonScriptPlugin/Private/PythonScriptPlugin.cpp
@@ -1200,6 +1200,8 @@ void FPythonScriptPlugin::RunStartupScripts()
 			// This allows them to be used to set-up an editor environment for the console
 			FPythonCommandEx InitUnrealPythonCommand;
 			InitUnrealPythonCommand.FileExecutionScope = EPythonFileExecutionScope::Public;
+			
+			FScopedDurationTimeLogger ScopedTimer(FString::Printf(TEXT("Running Python start-up script '%s'"), *PotentialFilePath));
 			RunFile(*PotentialFilePath, *InitUnrealPythonCommand.Command, InitUnrealPythonCommand);
 		}
 		else
@@ -1211,6 +1213,8 @@ void FPythonScriptPlugin::RunStartupScripts()
 	for (const FString& StartupScript : GetDefault<UPythonScriptPluginSettings>()->StartupScripts)
 	{
 		Progress.EnterProgressFrame(1.0f, FText::Format(LOCTEXT("PythonScriptPluginInitScripts_Running", "Running {0}..."), FText::FromString(StartupScript)));
+		
+		FScopedDurationTimeLogger ScopedTimer(FString::Printf(TEXT("Running Python start-up script '%s'"), *StartupScript));
 		ExecPythonCommand(*StartupScript);
 	}
 
