Only generate Python stub code when running full editor

This avoids an editor running with -game or -server, or a sub-editor running as a commandlet, from generating stub code

diff --git a/Engine/Plugins/Experimental/PythonScriptPlugin/Source/PythonScriptPlugin/Private/PythonScriptPlugin.cpp b/Engine/Plugins/Experimental/PythonScriptPlugin/Source/PythonScriptPlugin/Private/PythonScriptPlugin.cpp
--- a/Engine/Plugins/Experimental/PythonScriptPlugin/Source/PythonScriptPlugin/Private/PythonScriptPlugin.cpp
+++ b/Engine/Plugins/Experimental/PythonScriptPlugin/Source/PythonScriptPlugin/Private/PythonScriptPlugin.cpp
@@ -1164,7 +1164,7 @@ void FPythonScriptPlugin::GenerateStubCode()
 {
 	TRACE_CPUPROFILER_EVENT_SCOPE(FPythonScriptPlugin::GenerateStubCode)
 
-	if (IsDeveloperModeEnabled())
+	if (IsDeveloperModeEnabled() && GIsEditor && !IsRunningCommandlet())
 	{
 		// Generate stub code if developer mode enabled
 		FPyWrapperTypeRegistry::Get().GenerateStubCodeForWrappedTypes();
