Provided a common base implementation for the directory/file actions of a file content data source

#rnx

diff --git a/Engine/Plugins/Editor/ContentBrowser/ContentBrowserFileDataSource/Source/ContentBrowserFileDataSource/Private/ContentBrowserFileDataCore.cpp b/Engine/Plugins/Editor/ContentBrowser/ContentBrowserFileDataSource/Source/ContentBrowserFileDataSource/Private/ContentBrowserFileDataCore.cpp
--- a/Engine/Plugins/Editor/ContentBrowser/ContentBrowserFileDataSource/Source/ContentBrowserFileDataSource/Private/ContentBrowserFileDataCore.cpp
+++ b/Engine/Plugins/Editor/ContentBrowser/ContentBrowserFileDataSource/Source/ContentBrowserFileDataSource/Private/ContentBrowserFileDataCore.cpp
@@ -9,6 +9,7 @@
 #include "Misc/Paths.h"
 #include "AssetThumbnail.h"
 #include "AssetToolsModule.h"
+#include "AssetViewUtils.h"
 
 #define LOCTEXT_NAMESPACE "ContentBrowserFileDataSource"
 
@@ -1077,6 +1078,57 @@ bool GetFileItemAttributes(const FContentBrowserFileItemDataPayload& InFilePaylo
 	return false;
 }
 
+bool FDefaultFileActions::ItemPassesFilter(const FName InFilePath, const FString& InFilename, const FContentBrowserDataFilter& InFilter, const bool bIsFile)
+{
+	const FContentBrowserDataPackageFilter* PackageFilter = InFilter.ExtraFilters.FindFilter<FContentBrowserDataPackageFilter>();
+	if (PackageFilter && PackageFilter->PathPermissionList && PackageFilter->PathPermissionList->HasFiltering())
+	{
+		return PackageFilter->PathPermissionList->PassesStartsWithFilter(InFilePath, /*bAllowParentPaths*/!bIsFile);
+	}
+
+	return true;
+}
+
+bool FDefaultFileActions::GetItemAttribute(const FName InFilePath, const FString& InFilename, const bool InIncludeMetaData, const FName InAttributeKey, FContentBrowserItemDataAttributeValue& OutAttributeValue)
+{
+	if (InAttributeKey == ContentBrowserItemAttributes::ItemIsDeveloperContent)
+	{
+		const bool bIsDevelopersFolder = AssetViewUtils::IsDevelopersFolder(FNameBuilder(InFilePath));
+		OutAttributeValue.SetValue(bIsDevelopersFolder);
+		return true;
+	}
+
+	if (InAttributeKey == ContentBrowserItemAttributes::ItemIsLocalizedContent)
+	{
+		const bool bIsLocalizedFolder = FPackageName::IsLocalizedPackage(FNameBuilder(InFilePath));
+		OutAttributeValue.SetValue(bIsLocalizedFolder);
+		return true;
+	}
+
+	if (InAttributeKey == ContentBrowserItemAttributes::ItemIsEngineContent)
+	{
+		const bool bIsEngineFolder = AssetViewUtils::IsEngineFolder(FNameBuilder(InFilePath), /*bIncludePlugins*/true);
+		OutAttributeValue.SetValue(bIsEngineFolder);
+		return true;
+	}
+
+	if (InAttributeKey == ContentBrowserItemAttributes::ItemIsProjectContent)
+	{
+		const bool bIsProjectFolder = AssetViewUtils::IsProjectFolder(FNameBuilder(InFilePath), /*bIncludePlugins*/true);
+		OutAttributeValue.SetValue(bIsProjectFolder);
+		return true;
+	}
+
+	if (InAttributeKey == ContentBrowserItemAttributes::ItemIsPluginContent)
+	{
+		const bool bIsPluginFolder = AssetViewUtils::IsPluginFolder(FNameBuilder(InFilePath));
+		OutAttributeValue.SetValue(bIsPluginFolder);
+		return true;
+	}
+
+	return false;
+}
+
 }
 
 #undef LOCTEXT_NAMESPACE
diff --git a/Engine/Plugins/Editor/ContentBrowser/ContentBrowserFileDataSource/Source/ContentBrowserFileDataSource/Public/ContentBrowserFileDataCore.h b/Engine/Plugins/Editor/ContentBrowser/ContentBrowserFileDataSource/Source/ContentBrowserFileDataSource/Public/ContentBrowserFileDataCore.h
--- a/Engine/Plugins/Editor/ContentBrowser/ContentBrowserFileDataSource/Source/ContentBrowserFileDataSource/Public/ContentBrowserFileDataCore.h
+++ b/Engine/Plugins/Editor/ContentBrowser/ContentBrowserFileDataSource/Source/ContentBrowserFileDataSource/Public/ContentBrowserFileDataCore.h
@@ -4,6 +4,7 @@
 
 #include "ContentBrowserItemData.h"
 
+struct FContentBrowserDataFilter;
 class FContentBrowserFileItemDataPayload;
 class FContentBrowserFileItemDataPayload_Duplication;
 class FContentBrowserFolderItemDataPayload;
@@ -131,6 +132,13 @@ namespace ContentBrowserFileData
 
 	CONTENTBROWSERFILEDATASOURCE_API bool GetFileItemAttributes(const FContentBrowserFileItemDataPayload& InFilePayload, const bool InIncludeMetaData, FContentBrowserItemDataAttributeValues& OutAttributeValues);
 
+	class CONTENTBROWSERFILEDATASOURCE_API FDefaultFileActions
+	{
+	public:
+		static bool ItemPassesFilter(const FName InFilePath, const FString& InFilename, const FContentBrowserDataFilter& InFilter, const bool bIsFile);
+		static bool GetItemAttribute(const FName InFilePath, const FString& InFilename, const bool InIncludeMetaData, const FName InAttributeKey, FContentBrowserItemDataAttributeValue& OutAttributeValue);
+	};
+
 }
 
 #if UE_ENABLE_INCLUDE_ORDER_DEPRECATED_IN_5_2
diff --git a/Engine/Plugins/Experimental/PythonScriptPlugin/Source/PythonScriptPlugin/Private/PythonScriptPlugin.cpp b/Engine/Plugins/Experimental/PythonScriptPlugin/Source/PythonScriptPlugin/Private/PythonScriptPlugin.cpp
--- a/Engine/Plugins/Experimental/PythonScriptPlugin/Source/PythonScriptPlugin/Private/PythonScriptPlugin.cpp
+++ b/Engine/Plugins/Experimental/PythonScriptPlugin/Source/PythonScriptPlugin/Private/PythonScriptPlugin.cpp
@@ -954,59 +954,6 @@ void FPythonScriptPlugin::InitializePython()
 	{
 		ContentBrowserFileData::FFileConfigData PythonFileConfig;
 		{
-			auto PyItemPassesFilter = [](const FName InFilePath, const FString& InFilename, const FContentBrowserDataFilter& InFilter, const bool bIsFile)
-			{
-				const FContentBrowserDataPackageFilter* PackageFilter = InFilter.ExtraFilters.FindFilter<FContentBrowserDataPackageFilter>();
-				if (PackageFilter && PackageFilter->PathPermissionList && PackageFilter->PathPermissionList->HasFiltering())
-				{
-					return PackageFilter->PathPermissionList->PassesStartsWithFilter(InFilePath, /*bAllowParentPaths*/!bIsFile);
-				}
-
-				return true;
-			};
-
-			auto GetPyItemAttribute = [](const FName InFilePath, const FString& InFilename, const bool InIncludeMetaData, const FName InAttributeKey, FContentBrowserItemDataAttributeValue& OutAttributeValue)
-			{
-				// TODO: Need to way to avoid all this ToString() churn (FPackageNameView?)
-
-				if (InAttributeKey == ContentBrowserItemAttributes::ItemIsDeveloperContent)
-				{
-					const bool bIsDevelopersFolder = AssetViewUtils::IsDevelopersFolder(InFilePath.ToString());
-					OutAttributeValue.SetValue(bIsDevelopersFolder);
-					return true;
-				}
-
-				if (InAttributeKey == ContentBrowserItemAttributes::ItemIsLocalizedContent)
-				{
-					const bool bIsLocalizedFolder = FPackageName::IsLocalizedPackage(InFilePath.ToString());
-					OutAttributeValue.SetValue(bIsLocalizedFolder);
-					return true;
-				}
-
-				if (InAttributeKey == ContentBrowserItemAttributes::ItemIsEngineContent)
-				{
-					const bool bIsEngineFolder = AssetViewUtils::IsEngineFolder(InFilePath.ToString(), /*bIncludePlugins*/true);
-					OutAttributeValue.SetValue(bIsEngineFolder);
-					return true;
-				}
-
-				if (InAttributeKey == ContentBrowserItemAttributes::ItemIsProjectContent)
-				{
-					const bool bIsProjectFolder = AssetViewUtils::IsProjectFolder(InFilePath.ToString(), /*bIncludePlugins*/true);
-					OutAttributeValue.SetValue(bIsProjectFolder);
-					return true;
-				}
-
-				if (InAttributeKey == ContentBrowserItemAttributes::ItemIsPluginContent)
-				{
-					const bool bIsPluginFolder = AssetViewUtils::IsPluginFolder(InFilePath.ToString());
-					OutAttributeValue.SetValue(bIsPluginFolder);
-					return true;
-				}
-
-				return false;
-			};
-
 			auto PyItemPreview = [this](const FName InFilePath, const FString& InFilename)
 			{
 				ExecPythonCommand(*InFilename);
@@ -1014,8 +961,8 @@ void FPythonScriptPlugin::InitializePython()
 			};
 
 			ContentBrowserFileData::FDirectoryActions PyDirectoryActions;
-			PyDirectoryActions.PassesFilter.BindLambda(PyItemPassesFilter, false);
-			PyDirectoryActions.GetAttribute.BindLambda(GetPyItemAttribute);
+			PyDirectoryActions.PassesFilter.BindStatic(&ContentBrowserFileData::FDefaultFileActions::ItemPassesFilter, false);
+			PyDirectoryActions.GetAttribute.BindStatic(&ContentBrowserFileData::FDefaultFileActions::GetItemAttribute);
 			PythonFileConfig.SetDirectoryActions(PyDirectoryActions);
 
 			ContentBrowserFileData::FFileActions PyFileActions;
@@ -1026,8 +973,8 @@ void FPythonScriptPlugin::InitializePython()
 			PyFileActions.TypeFullDescription = LOCTEXT("PythonTypeFullDescription", "A file used to script the editor using Python");
 			PyFileActions.DefaultNewFileName = TEXT("new_python_script");
 			PyFileActions.TypeColor = FColor(255, 156, 0);
-			PyFileActions.PassesFilter.BindLambda(PyItemPassesFilter, true);
-			PyFileActions.GetAttribute.BindLambda(GetPyItemAttribute);
+			PyFileActions.PassesFilter.BindStatic(&ContentBrowserFileData::FDefaultFileActions::ItemPassesFilter, true);
+			PyFileActions.GetAttribute.BindStatic(&ContentBrowserFileData::FDefaultFileActions::GetItemAttribute);
 			PyFileActions.Preview.BindLambda(PyItemPreview);
 			PythonFileConfig.RegisterFileActions(PyFileActions);
 		}
