Force refresh the progress for the currently running Python script in case it hangs the editor

#rnx

diff --git a/Engine/Plugins/Experimental/PythonScriptPlugin/Source/PythonScriptPlugin/Private/PythonScriptPlugin.cpp b/Engine/Plugins/Experimental/PythonScriptPlugin/Source/PythonScriptPlugin/Private/PythonScriptPlugin.cpp
--- a/Engine/Plugins/Experimental/PythonScriptPlugin/Source/PythonScriptPlugin/Private/PythonScriptPlugin.cpp
+++ b/Engine/Plugins/Experimental/PythonScriptPlugin/Source/PythonScriptPlugin/Private/PythonScriptPlugin.cpp
@@ -1142,6 +1142,7 @@ void FPythonScriptPlugin::RunStartupScripts()
 		if (FPaths::FileExists(PotentialFilePath))
 		{
 			Progress.EnterProgressFrame(1.0f, FText::Format(LOCTEXT("PythonScriptPluginInitScripts_Running", "Running {0}..."), FText::FromString(PotentialFilePath)));
+			Progress.ForceRefresh();
 
 			// Execute these files in the "public" scope, as if their contents had been run directly in the console
 			// This allows them to be used to set-up an editor environment for the console
