Fix PyReferenceCollector/ScriptDelegates regression

diff --git a/Engine/Plugins/Experimental/PythonScriptPlugin/Source/PythonScriptPlugin/Private/PyReferenceCollector.cpp b/Engine/Plugins/Experimental/PythonScriptPlugin/Source/PythonScriptPlugin/Private/PyReferenceCollector.cpp
--- a/Engine/Plugins/Experimental/PythonScriptPlugin/Source/PythonScriptPlugin/Private/PyReferenceCollector.cpp
+++ b/Engine/Plugins/Experimental/PythonScriptPlugin/Source/PythonScriptPlugin/Private/PyReferenceCollector.cpp
@@ -188,8 +188,9 @@ void FPyReferenceCollector::AddReferencedObjectsFromDelegate(FReferenceCollector
 {
 	// Keep the delegate object alive if it's using a Python proxy instance
 	// We have to use the EvenIfUnreachable variant here as the objects are speculatively marked as unreachable during GC
-	if (auto& Ptr = InDelegate.GetUObjectRefEvenIfUnreachable(); Cast<UPythonCallableForDelegate>(Ptr.GetEvenIfUnreachable()))
+	if (UPythonCallableForDelegate* PythonCallableForDelegate = Cast<UPythonCallableForDelegate>(InDelegate.GetUObjectEvenIfUnreachable()))
 	{
+		TWeakObjectPtr<UPythonCallableForDelegate> Ptr{PythonCallableForDelegate};
 		InCollector.AddReferencedObject(Ptr);
 	}
 }
@@ -198,11 +199,12 @@ void FPyReferenceCollector::AddReferencedObjectsFromMulticastDelegate(FReference
 {
 	// Keep the delegate objects alive if they're using a Python proxy instance
 	// We have to use the EvenIfUnreachable variant here as the objects are speculatively marked as unreachable during GC
-	for (auto* DelegateObj : InDelegate.GetAllObjectRefsEvenIfUnreachable())
+	for (UObject* DelegateObj : InDelegate.GetAllObjectsEvenIfUnreachable())
 	{
-		if (Cast<UPythonCallableForDelegate>(DelegateObj->Get()))
+		if (UPythonCallableForDelegate* PythonCallableForDelegate = Cast<UPythonCallableForDelegate>(DelegateObj))
 		{
-			InCollector.AddReferencedObject(*DelegateObj);
+			TWeakObjectPtr<UPythonCallableForDelegate> Ptr{PythonCallableForDelegate};
+			InCollector.AddReferencedObject(Ptr);
 		}
 	}
 }
diff --git a/Engine/Source/Runtime/Core/Public/UObject/ScriptDelegates.h b/Engine/Source/Runtime/Core/Public/UObject/ScriptDelegates.h
--- a/Engine/Source/Runtime/Core/Public/UObject/ScriptDelegates.h
+++ b/Engine/Source/Runtime/Core/Public/UObject/ScriptDelegates.h
@@ -394,7 +394,7 @@ public:
 		return static_cast< const UObject* >( Object.GetEvenIfUnreachable() );
 	}
 
-	WeakPtrType& GetUObjectRefEvenIfUnreachable()
+	WeakPtrType& GetUObjectRef()
 	{
 		return Object;
 	}
@@ -945,7 +945,7 @@ public:
 		TArray<UObject*> Result;
 		for (auto* Ref : GetAllObjectRefsEvenIfUnreachable())
 		{
-			Result.Add(*Ref);
+			Result.Add(Ref->GetEvenIfUnreachable());
 		}
 		return Result;
 	}
@@ -957,8 +957,8 @@ public:
 		TArray< WeakPtrType* > OutputList;
 		for( typename InvocationListType::TIterator CurDelegate( InvocationList ); CurDelegate; ++CurDelegate )
 		{
-			WeakPtrType& CurObject = CurDelegate->GetUObjectRefEvenIfUnreachable();
-			if( CurObject != nullptr )
+			WeakPtrType& CurObject = CurDelegate->GetUObjectRef();
+			if( CurObject.GetEvenIfUnreachable() != nullptr )
 			{
 				OutputList.Add( &CurObject );
 			}
