Allow Python to be disabled entirely for certain commandlets

This can be done by adding the following to any Editor.ini file:
  [PythonScriptPlugin]
  +DisablePythonForCommandlet=CommandletName

diff --git a/Engine/Plugins/Experimental/PythonScriptPlugin/Source/PythonScriptPlugin/Private/PythonScriptPlugin.cpp b/Engine/Plugins/Experimental/PythonScriptPlugin/Source/PythonScriptPlugin/Private/PythonScriptPlugin.cpp
--- a/Engine/Plugins/Experimental/PythonScriptPlugin/Source/PythonScriptPlugin/Private/PythonScriptPlugin.cpp
+++ b/Engine/Plugins/Experimental/PythonScriptPlugin/Source/PythonScriptPlugin/Private/PythonScriptPlugin.cpp
@@ -489,8 +489,52 @@ FPythonScriptPlugin::FPythonScriptPlugin()
 
 bool FPythonScriptPlugin::IsPythonAvailable() const
 {
-	static const bool bDisablePython = FParse::Param(FCommandLine::Get(), TEXT("DisablePython"));
-	return WITH_PYTHON && !bDisablePython;
+#if WITH_PYTHON
+	auto IsPythonEnabled = []()
+	{
+		if (FParse::Param(FCommandLine::Get(), TEXT("DisablePython")))
+		{
+			UE_LOG(LogPython, Log, TEXT("Python disabled via command-line flag '-DisablePython'"));
+			return false;
+		}
+
+		if (IsRunningCommandlet())
+		{
+			TArray<FString> DisablePythonForCommandlets;
+			GConfig->GetArray(TEXT("PythonScriptPlugin"), TEXT("DisablePythonForCommandlet"), DisablePythonForCommandlets, GEditorIni);
+
+			FString RunningCommandletName;
+			if (!DisablePythonForCommandlets.IsEmpty() && FParse::Value(FCommandLine::Get(), TEXT("-run="), RunningCommandletName))
+			{
+				auto CleanCommandletName = [](FString& InOutCommandletName)
+				{
+					const FStringView CommandletSuffix = TEXTVIEW("Commandlet");
+					if (InOutCommandletName.EndsWith(CommandletSuffix))
+					{
+						InOutCommandletName.LeftChopInline(CommandletSuffix.Len(), /*bAllowShrinking*/false);
+					}
+				};
+
+				CleanCommandletName(RunningCommandletName);
+				for (FString& DisablePythonForCommandlet : DisablePythonForCommandlets)
+				{
+					CleanCommandletName(DisablePythonForCommandlet);
+					if (DisablePythonForCommandlet == RunningCommandletName)
+					{
+						UE_LOG(LogPython, Log, TEXT("Python disabled via config setting 'DisablePythonForCommandlet'"));
+						return false;
+					}
+				}
+			}
+		}
+
+		return true;
+	};
+	static const bool bEnablePython = IsPythonEnabled();
+	return bEnablePython;
+#else
+	return false;
+#endif
 }
 
 bool FPythonScriptPlugin::ExecPythonCommand(const TCHAR* InPythonCommand)
