Fixed Python Nativization failing to validate against an interface class

diff --git a/Engine/Plugins/Experimental/PythonScriptPlugin/Source/PythonScriptPlugin/Private/PyConversion.cpp b/Engine/Plugins/Experimental/PythonScriptPlugin/Source/PythonScriptPlugin/Private/PyConversion.cpp
--- a/Engine/Plugins/Experimental/PythonScriptPlugin/Source/PythonScriptPlugin/Private/PyConversion.cpp
+++ b/Engine/Plugins/Experimental/PythonScriptPlugin/Source/PythonScriptPlugin/Private/PyConversion.cpp
@@ -476,10 +476,22 @@ FPyConversionResult Pythonize(UObject* Val, PyObject*& OutPyObj, const ESetError
 
 FPyConversionResult NativizeObject(PyObject* PyObj, UObject*& OutVal, UClass* ExpectedType, const ESetErrorState SetErrorState)
 {
+	auto IsObjectExpectedType = [ExpectedType](const UObject* ObjectInstance) -> bool
+	{
+		if (!ExpectedType)
+		{
+			return true;
+		}
+
+		return ExpectedType->HasAnyClassFlags(CLASS_Interface)
+			? ObjectInstance->GetClass()->ImplementsInterface(ExpectedType)
+			: ObjectInstance->IsA(ExpectedType);
+	};
+
 	if (PyObject_IsInstance(PyObj, (PyObject*)&PyWrapperObjectType) == 1)
 	{
 		FPyWrapperObject* PyWrappedObj = (FPyWrapperObject*)PyObj;
-		if (FPyWrapperObject::ValidateInternalState(PyWrappedObj) && (!ExpectedType || PyWrappedObj->ObjectInstance->IsA(ExpectedType)))
+		if (FPyWrapperObject::ValidateInternalState(PyWrappedObj) && IsObjectExpectedType(PyWrappedObj->ObjectInstance))
 		{
 			OutVal = PyWrappedObj->ObjectInstance;
 			return FPyConversionResult::Success();
