Updated get/set_editor_property in Python to handle getting/setting the sparse class data for a CDO

This proxies to the sparse class data if the property is missing or deprecated on the object class, as is the case for the BP version of these functions

diff --git a/Engine/Plugins/Experimental/PythonScriptPlugin/Source/PythonScriptPlugin/Private/PyTest.h b/Engine/Plugins/Experimental/PythonScriptPlugin/Source/PythonScriptPlugin/Private/PyTest.h
--- a/Engine/Plugins/Experimental/PythonScriptPlugin/Source/PythonScriptPlugin/Private/PyTest.h
+++ b/Engine/Plugins/Experimental/PythonScriptPlugin/Source/PythonScriptPlugin/Private/PyTest.h
@@ -141,10 +141,23 @@ class UPyTestStructLibrary : public UBlueprintFunctionLibrary
 	static void ClearBoolMutableViaRef(UPARAM(ref) FPyTestStruct& InStruct);
 };
 
+/**
+ * Struct to allow testing of class sparse data on a Python exposed type.
+ */
+USTRUCT(BlueprintType)
+struct FPyTestClassSparseData
+{
+	GENERATED_BODY()
+
+public:
+	UPROPERTY(EditDefaultsOnly, BlueprintReadOnly, Category = "Python|Internal")
+	int32 IntFromSparseData = 0;
+};
+
 /**
  * Object to allow testing of the various UObject features that are exposed to Python wrapped types.
  */
-UCLASS(Blueprintable)
+UCLASS(Blueprintable, SparseClassDataTypes = PyTestClassSparseData)
 class UPyTestObject : public UObject, public IPyTestChildInterface, public IPyTestOtherInterface
 {
 	GENERATED_BODY()
diff --git a/Engine/Plugins/Experimental/PythonScriptPlugin/Source/PythonScriptPlugin/Private/PyWrapperObject.cpp b/Engine/Plugins/Experimental/PythonScriptPlugin/Source/PythonScriptPlugin/Private/PyWrapperObject.cpp
--- a/Engine/Plugins/Experimental/PythonScriptPlugin/Source/PythonScriptPlugin/Private/PyWrapperObject.cpp
+++ b/Engine/Plugins/Experimental/PythonScriptPlugin/Source/PythonScriptPlugin/Private/PyWrapperObject.cpp
@@ -740,16 +740,32 @@ PyTypeObject InitializePyWrapperObjectType()
 			PyGenUtil::FGeneratedWrappedProperty WrappedPropDef;
 
 			const UClass* Class = InSelf->ObjectInstance->GetClass();
-
 			const FName ResolvedName = FPyWrapperObjectMetaData::ResolvePropertyName(InSelf, InPythonPropName);
+
 			const FProperty* ResolvedProp = Class->FindPropertyByName(ResolvedName);
+			if (InSelf->ObjectInstance->HasAnyFlags(RF_ClassDefaultObject) && (!ResolvedProp || ResolvedProp->HasAnyPropertyFlags(CPF_Deprecated)))
+			{
+				// Look for a sparse member of the same name - the sparse data is treated as an extension of the class default object by the details panel
+				if (const UStruct* SparseDataStruct = Class->GetSparseClassDataStruct())
+				{
+					if (const FProperty* SparseProp = SparseDataStruct->FindPropertyByName(ResolvedName))
+					{
+						ResolvedProp = SparseProp;
+					}
+				}
+			}
+
 			if (!ResolvedProp)
 			{
 				PyUtil::SetPythonError(PyExc_Exception, InSelf, *FString::Printf(TEXT("Failed to find property '%s' for attribute '%s' on '%s'"), *ResolvedName.ToString(), *InPythonPropName.ToString(), *Class->GetName()));
 				return WrappedPropDef;
 			}
 
+			// If the owner class is set then this property is from 'self', otherwise it's from the sparse class data
+			const bool bPropertyIsOwnedBySelf = ResolvedProp->GetOwnerClass() != nullptr;
+
 			TOptional<FString> PropDeprecationMessage;
+			if (bPropertyIsOwnedBySelf)
 			{
 				FString PropDeprecationMessageStr;
 				if (FPyWrapperObjectMetaData::IsPropertyDeprecated(InSelf, InPythonPropName, &PropDeprecationMessageStr))
@@ -799,7 +815,27 @@ PyTypeObject InitializePyWrapperObjectType()
 				return nullptr;
 			}
 
-			return FPyWrapperObject::GetPropertyValue(InSelf, WrappedPropDef, TCHAR_TO_UTF8(*Name.ToString()));
+			// If the owner class is set then this property is from 'self', otherwise it's from the sparse class data
+			const bool bPropertyIsOwnedBySelf = WrappedPropDef.Prop->GetOwnerClass() != nullptr;
+			UClass* Class = InSelf->ObjectInstance->GetClass();
+			if (bPropertyIsOwnedBySelf)
+			{
+				return PyGenUtil::GetPropertyValue(Class, InSelf->ObjectInstance, WrappedPropDef, TCHAR_TO_UTF8(*Name.ToString()), (PyObject*)InSelf, *PyUtil::GetErrorContext(InSelf));
+			}
+			else
+			{
+				const UScriptStruct* SparseDataStruct = Class->GetSparseClassDataStruct();
+				const void* SparseData = Class->GetSparseClassData(EGetSparseClassDataMethod::ArchetypeIfNull);
+				if (SparseDataStruct && SparseData)
+				{
+					return PyGenUtil::GetPropertyValue(SparseDataStruct, SparseData, WrappedPropDef, TCHAR_TO_UTF8(*Name.ToString()), (PyObject*)InSelf, *PyUtil::GetErrorContext(InSelf));
+				}
+				else
+				{
+					PyUtil::SetPythonError(PyExc_Exception, InSelf, *FString::Printf(TEXT("Failed to get sparse data for '%s'"), *Class->GetName()));
+					return nullptr;
+				}
+			}
 		}
 
 		static PyObject* SetEditorProperty(FPyWrapperObject* InSelf, PyObject* InArgs, PyObject* InKwds)
@@ -841,10 +877,36 @@ PyTypeObject InitializePyWrapperObjectType()
 				return nullptr;
 			}
 
-			const int Result = FPyWrapperObject::SetPropertyValue(InSelf, PyValueObj, WrappedPropDef, TCHAR_TO_UTF8(*Name.ToString()), NotifyMode, PropertyAccessUtil::EditorReadOnlyFlags);
-			if (Result != 0)
+			// If the owner class is set then this property is from 'self', otherwise it's from the sparse class data
+			const bool bPropertyIsOwnedBySelf = WrappedPropDef.Prop->GetOwnerClass() != nullptr;
+			UClass* Class = InSelf->ObjectInstance->GetClass();
+			if (bPropertyIsOwnedBySelf)
 			{
-				return nullptr;
+				const TUniquePtr<FPropertyAccessChangeNotify> ChangeNotify = FPyWrapperOwnerContext((PyObject*)InSelf, WrappedPropDef.Prop).BuildChangeNotify(NotifyMode);
+				const int Result = PyGenUtil::SetPropertyValue(Class, InSelf->ObjectInstance, PyValueObj, WrappedPropDef, TCHAR_TO_UTF8(*Name.ToString()), ChangeNotify.Get(), PropertyAccessUtil::EditorReadOnlyFlags, PropertyAccessUtil::IsObjectTemplate(InSelf->ObjectInstance), *PyUtil::GetErrorContext(InSelf));
+				if (Result != 0)
+				{
+					return nullptr;
+				}
+			}
+			else
+			{
+				const UScriptStruct* SparseDataStruct = Class->GetSparseClassDataStruct();
+				void* SparseData = Class->GetOrCreateSparseClassData();
+				if (SparseDataStruct && SparseData)
+				{
+					const TUniquePtr<FPropertyAccessChangeNotify> ChangeNotify = FPyWrapperOwnerContext((PyObject*)InSelf, WrappedPropDef.Prop).BuildChangeNotify(NotifyMode);
+					const int Result = PyGenUtil::SetPropertyValue(SparseDataStruct, SparseData, PyValueObj, WrappedPropDef, TCHAR_TO_UTF8(*Name.ToString()), ChangeNotify.Get(), PropertyAccessUtil::EditorReadOnlyFlags, PropertyAccessUtil::IsObjectTemplate(InSelf->ObjectInstance), *PyUtil::GetErrorContext(InSelf));
+					if (Result != 0)
+					{
+						return nullptr;
+					}
+				}
+				else
+				{
+					PyUtil::SetPythonError(PyExc_Exception, InSelf, *FString::Printf(TEXT("Failed to get sparse data for '%s'"), *Class->GetName()));
+					return nullptr;
+				}
 			}
 
 			Py_RETURN_NONE;
@@ -916,16 +978,40 @@ PyTypeObject InitializePyWrapperObjectType()
 			};
 
 			// Try and set the value of each property
+			UClass* Class = InSelf->ObjectInstance->GetClass();
 			for (const FPropertyInfoPair& PropertyInfo : PropertyInfos)
 			{
 				const FName Name = PropertyInfo.Get<0>();
 				const PyGenUtil::FGeneratedWrappedProperty& WrappedPropDef = PropertyInfo.Get<1>();
 				PyObject* PyValueObj = PropertyInfo.Get<2>().GetPtr();
 
-				const int Result = FPyWrapperObject::SetPropertyValue(InSelf, PyValueObj, WrappedPropDef, TCHAR_TO_UTF8(*Name.ToString()), EPropertyAccessChangeNotifyMode::Never, PropertyAccessUtil::EditorReadOnlyFlags);
-				if (Result != 0)
+				// If the owner class is set then this property is from 'self', otherwise it's from the sparse class data
+				const bool bPropertyIsOwnedBySelf = WrappedPropDef.Prop->GetOwnerClass() != nullptr;
+				if (bPropertyIsOwnedBySelf)
 				{
-					return nullptr;
+					const int Result = PyGenUtil::SetPropertyValue(Class, InSelf->ObjectInstance, PyValueObj, WrappedPropDef, TCHAR_TO_UTF8(*Name.ToString()), nullptr, PropertyAccessUtil::EditorReadOnlyFlags, PropertyAccessUtil::IsObjectTemplate(InSelf->ObjectInstance), *PyUtil::GetErrorContext(InSelf));
+					if (Result != 0)
+					{
+						return nullptr;
+					}
+				}
+				else
+				{
+					const UScriptStruct* SparseDataStruct = Class->GetSparseClassDataStruct();
+					void* SparseData = Class->GetOrCreateSparseClassData();
+					if (SparseDataStruct && SparseData)
+					{
+						const int Result = PyGenUtil::SetPropertyValue(SparseDataStruct, SparseData, PyValueObj, WrappedPropDef, TCHAR_TO_UTF8(*Name.ToString()), nullptr, PropertyAccessUtil::EditorReadOnlyFlags, PropertyAccessUtil::IsObjectTemplate(InSelf->ObjectInstance), *PyUtil::GetErrorContext(InSelf));
+						if (Result != 0)
+						{
+							return nullptr;
+						}
+					}
+					else
+					{
+						PyUtil::SetPythonError(PyExc_Exception, InSelf, *FString::Printf(TEXT("Failed to get sparse data for '%s'"), *Class->GetName()));
+						return nullptr;
+					}
 				}
 			}
 
