Avoid using FStructOnScope when calling UFunctions from Python

#rnx

diff --git a/Engine/Plugins/Experimental/PythonScriptPlugin/Source/PythonScriptPlugin/Private/PyGenUtil.cpp b/Engine/Plugins/Experimental/PythonScriptPlugin/Source/PythonScriptPlugin/Private/PyGenUtil.cpp
--- a/Engine/Plugins/Experimental/PythonScriptPlugin/Source/PythonScriptPlugin/Private/PyGenUtil.cpp
+++ b/Engine/Plugins/Experimental/PythonScriptPlugin/Source/PythonScriptPlugin/Private/PyGenUtil.cpp
@@ -448,12 +448,12 @@ void FGeneratedWrappedConstant::ToPython(FPyConstantDef& OutPyConstant) const
 			}
 	
 			// Return value requires that we create a params struct to hold the result
-			FStructOnScope FuncParams(This->ConstantFunc.Func);
-			if (!PyUtil::InvokeFunctionCall(Obj, This->ConstantFunc.Func, FuncParams.GetStructMemory(), *ErrorCtxt))
+			PY_UFUNCTION_STACK(FuncParams, This->ConstantFunc.Func);
+			if (!PyUtil::InvokeFunctionCall(Obj, This->ConstantFunc.Func, FuncParams.GetMemory(), *ErrorCtxt))
 			{
 				return nullptr;
 			}
-			return PyGenUtil::PackReturnValues(FuncParams.GetStructMemory(), This->ConstantFunc.OutputParams, *ErrorCtxt, *FString::Printf(TEXT("constant '%s' on '%s'"), UTF8_TO_TCHAR(This->ConstantName.GetData()), *PyUtil::GetCleanTypename(InType)));
+			return PyGenUtil::PackReturnValues(FuncParams.GetMemory(), This->ConstantFunc.OutputParams, *ErrorCtxt, *FString::Printf(TEXT("constant '%s' on '%s'"), UTF8_TO_TCHAR(This->ConstantName.GetData()), *PyUtil::GetCleanTypename(InType)));
 		}
 	
 		Py_RETURN_NONE;
@@ -2412,17 +2412,17 @@ void PythonizeStructValueImpl(const UScriptStruct* InStruct, const void* InStruc
 			if (BreakFuncDef.OutputParams.Num() <= 255)
 			{
 				// Call the break function using the instance we were given
-				FStructOnScope FuncParams(BreakFuncDef.Func);
+				PY_UFUNCTION_STACK(FuncParams, BreakFuncDef.Func);
 				if (BreakFuncDef.InputParams.Num() == 1 && CastField<FStructProperty>(BreakFuncDef.InputParams[0].ParamProp) && InStruct->IsChildOf(CastFieldChecked<const FStructProperty>(BreakFuncDef.InputParams[0].ParamProp)->Struct))
 				{
 					// Copy the given instance as the 'self' argument
 					const FGeneratedWrappedMethodParameter& SelfParam = BreakFuncDef.InputParams[0];
-					void* SelfArgInstance = SelfParam.ParamProp->ContainerPtrToValuePtr<void>(FuncParams.GetStructMemory());
+					void* SelfArgInstance = SelfParam.ParamProp->ContainerPtrToValuePtr<void>(FuncParams.GetMemory());
 					CastFieldChecked<const FStructProperty>(SelfParam.ParamProp)->Struct->CopyScriptStruct(SelfArgInstance, InStructValue);
 				}
 				{
 					FPyScopedGIL GIL;
-					PyUtil::InvokeFunctionCall(Obj, BreakFuncDef.Func, FuncParams.GetStructMemory(), TEXT("pythonize default struct value"));
+					PyUtil::InvokeFunctionCall(Obj, BreakFuncDef.Func, FuncParams.GetMemory(), TEXT("pythonize default struct value"));
 					PyErr_Clear(); // Clear any errors in case InvokeFunctionCall failed
 				}
 
@@ -2434,7 +2434,7 @@ void PythonizeStructValueImpl(const UScriptStruct* InStruct, const void* InStruc
 					{
 						OutPythonDefaultValue += TEXT(", ");
 					}
-					PythonizeValueImpl(OutputParam.ParamProp, OutputParam.ParamProp->ContainerPtrToValuePtr<void>(FuncParams.GetStructMemory()), InFlags, OutPythonDefaultValue);
+					PythonizeValueImpl(OutputParam.ParamProp, OutputParam.ParamProp->ContainerPtrToValuePtr<void>(FuncParams.GetMemory()), InFlags, OutPythonDefaultValue);
 				}
 			}
 		}
diff --git a/Engine/Plugins/Experimental/PythonScriptPlugin/Source/PythonScriptPlugin/Private/PyUtil.h b/Engine/Plugins/Experimental/PythonScriptPlugin/Source/PythonScriptPlugin/Private/PyUtil.h
--- a/Engine/Plugins/Experimental/PythonScriptPlugin/Source/PythonScriptPlugin/Private/PyUtil.h
+++ b/Engine/Plugins/Experimental/PythonScriptPlugin/Source/PythonScriptPlugin/Private/PyUtil.h
@@ -8,6 +8,7 @@
 #include "Containers/ArrayView.h"
 #include "Logging/LogMacros.h"
 #include "UObject/Field.h"
+#include "UObject/FieldIterator.h"
 
 DECLARE_LOG_CATEGORY_EXTERN(LogPython, Log, All);
 
@@ -264,6 +265,45 @@ namespace PyUtil
 		explicit FMapValueOnScope(const FMapProperty* InProp);
 	};
 
+	/** Helper used to hold the stack data for a UFunction call */
+	struct FFunctionStackOnScope
+	{
+	public:
+		FFunctionStackOnScope(const UFunction* InFunc, void* InStackMemory)
+			: Func(InFunc)
+			, StackMemory(InStackMemory)
+		{
+			for (TFieldIterator<FProperty> PropIt(Func); PropIt; ++PropIt)
+			{
+				PropIt->InitializeValue_InContainer(StackMemory);
+			}
+		}
+
+		~FFunctionStackOnScope()
+		{
+			for (TFieldIterator<FProperty> PropIt(Func); PropIt; ++PropIt)
+			{
+				PropIt->DestroyValue_InContainer(StackMemory);
+			}
+		}
+
+		const UFunction* GetFunction() const
+		{
+			return Func;
+		}
+
+		void* GetMemory() const
+		{
+			return StackMemory;
+		}
+
+	private:
+		const UFunction* Func = nullptr;
+		void* StackMemory = nullptr;
+	};
+	#define PY_UFUNCTION_STACK(NAME, FUNCTION) \
+		PyUtil::FFunctionStackOnScope NAME(FUNCTION, FMemory_Alloca_Aligned(((FUNCTION)->GetStructureSize() ? (FUNCTION)->GetStructureSize() : 1), (FUNCTION)->GetMinAlignment()))
+
 	/** Struct containing information needed to construct a property instance */
 	struct FPropertyDef
 	{
diff --git a/Engine/Plugins/Experimental/PythonScriptPlugin/Source/PythonScriptPlugin/Private/PyWrapperDelegate.cpp b/Engine/Plugins/Experimental/PythonScriptPlugin/Source/PythonScriptPlugin/Private/PyWrapperDelegate.cpp
--- a/Engine/Plugins/Experimental/PythonScriptPlugin/Source/PythonScriptPlugin/Private/PyWrapperDelegate.cpp
+++ b/Engine/Plugins/Experimental/PythonScriptPlugin/Source/PythonScriptPlugin/Private/PyWrapperDelegate.cpp
@@ -398,8 +398,8 @@ struct TPyWrapperDelegateImpl
 			return nullptr;
 		}
 
-		FStructOnScope DelegateParams(DelegateSignature.Func);
-		PyGenUtil::ApplyParamDefaults(DelegateParams.GetStructMemory(), DelegateSignature.InputParams);
+		PY_UFUNCTION_STACK(DelegateParams, DelegateSignature.Func);
+		PyGenUtil::ApplyParamDefaults(DelegateParams.GetMemory(), DelegateSignature.InputParams);
 		for (int32 ParamIndex = 0; ParamIndex < Params.Num(); ++ParamIndex)
 		{
 			const PyGenUtil::FGeneratedWrappedMethodParameter& ParamDef = DelegateSignature.InputParams[ParamIndex];
@@ -407,15 +407,15 @@ struct TPyWrapperDelegateImpl
 			PyObject* PyValue = Params[ParamIndex];
 			if (PyValue)
 			{
-				if (!PyConversion::NativizeProperty_InContainer(PyValue, ParamDef.ParamProp, DelegateParams.GetStructMemory(), 0))
+				if (!PyConversion::NativizeProperty_InContainer(PyValue, ParamDef.ParamProp, DelegateParams.GetMemory(), 0))
 				{
 					PyUtil::SetPythonError(PyExc_TypeError, InSelf, *FString::Printf(TEXT("Failed to convert parameter '%s' when calling delegate"), UTF8_TO_TCHAR(ParamDef.ParamName.GetData())));
 					return nullptr;
 				}
 			}
 		}
-		FDelegateInvocation::Call(*InSelf->DelegateInstance, DelegateParams.GetStructMemory());
-		return PyGenUtil::PackReturnValues(DelegateParams.GetStructMemory(), DelegateSignature.OutputParams, *PyUtil::GetErrorContext(InSelf), TEXT("delegate"));
+		FDelegateInvocation::Call(*InSelf->DelegateInstance, DelegateParams.GetMemory());
+		return PyGenUtil::PackReturnValues(DelegateParams.GetMemory(), DelegateSignature.OutputParams, *PyUtil::GetErrorContext(InSelf), TEXT("delegate"));
 	}
 };
 
diff --git a/Engine/Plugins/Experimental/PythonScriptPlugin/Source/PythonScriptPlugin/Private/PyWrapperObject.cpp b/Engine/Plugins/Experimental/PythonScriptPlugin/Source/PythonScriptPlugin/Private/PyWrapperObject.cpp
--- a/Engine/Plugins/Experimental/PythonScriptPlugin/Source/PythonScriptPlugin/Private/PyWrapperObject.cpp
+++ b/Engine/Plugins/Experimental/PythonScriptPlugin/Source/PythonScriptPlugin/Private/PyWrapperObject.cpp
@@ -175,16 +175,16 @@ int FPyWrapperObject::CallSetterFunction(FPyWrapperObject* InSelf, PyObject* InV
 			return -1;
 		}
 
-		FStructOnScope FuncParams(InFuncDef.Func);
+		PY_UFUNCTION_STACK(FuncParams, InFuncDef.Func);
 		if (InValue)
 		{
-			if (!PyConversion::NativizeProperty_InContainer(InValue, InFuncDef.InputParams[0].ParamProp, FuncParams.GetStructMemory(), 0))
+			if (!PyConversion::NativizeProperty_InContainer(InValue, InFuncDef.InputParams[0].ParamProp, FuncParams.GetMemory(), 0))
 			{
 				PyUtil::SetPythonError(PyExc_TypeError, InSelf, *FString::Printf(TEXT("Failed to convert input parameter when calling function '%s.%s' on '%s'"), *InFuncDef.Func->GetOwnerClass()->GetName(), *InFuncDef.Func->GetName(), *InSelf->ObjectInstance->GetName()));
 				return -1;
 			}
 		}
-		if (!PyUtil::InvokeFunctionCall(InSelf->ObjectInstance, InFuncDef.Func, FuncParams.GetStructMemory(), *PyUtil::GetErrorContext(InSelf)))
+		if (!PyUtil::InvokeFunctionCall(InSelf->ObjectInstance, InFuncDef.Func, FuncParams.GetMemory(), *PyUtil::GetErrorContext(InSelf)))
 		{
 			return -1;
 		}
@@ -252,12 +252,12 @@ PyObject* FPyWrapperObject::CallFunction_Impl(UObject* InObj, const PyGenUtil::F
 		else
 		{
 			// Return value requires that we create a params struct to hold the result
-			FStructOnScope FuncParams(InFuncDef.Func);
-			if (!PyUtil::InvokeFunctionCall(InObj, InFuncDef.Func, FuncParams.GetStructMemory(), InErrorCtxt))
+			PY_UFUNCTION_STACK(FuncParams, InFuncDef.Func);
+			if (!PyUtil::InvokeFunctionCall(InObj, InFuncDef.Func, FuncParams.GetMemory(), InErrorCtxt))
 			{
 				return nullptr;
 			}
-			return PyGenUtil::PackReturnValues(FuncParams.GetStructMemory(), InFuncDef.OutputParams, InErrorCtxt, *FString::Printf(TEXT("function '%s.%s' on '%s'"), *InFuncDef.Func->GetOwnerClass()->GetName(), *InFuncDef.Func->GetName(), *InObj->GetName()));
+			return PyGenUtil::PackReturnValues(FuncParams.GetMemory(), InFuncDef.OutputParams, InErrorCtxt, *FString::Printf(TEXT("function '%s.%s' on '%s'"), *InFuncDef.Func->GetOwnerClass()->GetName(), *InFuncDef.Func->GetName(), *InObj->GetName()));
 		}
 	}
 
@@ -284,8 +284,8 @@ PyObject* FPyWrapperObject::CallFunction_Impl(UObject* InObj, PyObject* InArgs,
 			}
 		}
 
-		FStructOnScope FuncParams(InFuncDef.Func);
-		PyGenUtil::ApplyParamDefaults(FuncParams.GetStructMemory(), InFuncDef.InputParams);
+		PY_UFUNCTION_STACK(FuncParams, InFuncDef.Func);
+		PyGenUtil::ApplyParamDefaults(FuncParams.GetMemory(), InFuncDef.InputParams);
 		for (int32 ParamIndex = 0; ParamIndex < Params.Num(); ++ParamIndex)
 		{
 			const PyGenUtil::FGeneratedWrappedMethodParameter& ParamDef = InFuncDef.InputParams[ParamIndex];
@@ -293,18 +293,18 @@ PyObject* FPyWrapperObject::CallFunction_Impl(UObject* InObj, PyObject* InArgs,
 			PyObject* PyValue = Params[ParamIndex];
 			if (PyValue)
 			{
-				if (!PyConversion::NativizeProperty_InContainer(PyValue, ParamDef.ParamProp, FuncParams.GetStructMemory(), 0))
+				if (!PyConversion::NativizeProperty_InContainer(PyValue, ParamDef.ParamProp, FuncParams.GetMemory(), 0))
 				{
 					PyUtil::SetPythonError(PyExc_TypeError, InErrorCtxt, *FString::Printf(TEXT("Failed to convert parameter '%s' when calling function '%s.%s' on '%s'"), UTF8_TO_TCHAR(ParamDef.ParamName.GetData()), *InFuncDef.Func->GetOwnerClass()->GetName(), *InFuncDef.Func->GetName(), *InObj->GetName()));
 					return nullptr;
 				}
 			}
 		}
-		if (!PyUtil::InvokeFunctionCall(InObj, InFuncDef.Func, FuncParams.GetStructMemory(), InErrorCtxt))
+		if (!PyUtil::InvokeFunctionCall(InObj, InFuncDef.Func, FuncParams.GetMemory(), InErrorCtxt))
 		{
 			return nullptr;
 		}
-		return PyGenUtil::PackReturnValues(FuncParams.GetStructMemory(), InFuncDef.OutputParams, InErrorCtxt, *FString::Printf(TEXT("function '%s.%s' on '%s'"), *InFuncDef.Func->GetOwnerClass()->GetName(), *InFuncDef.Func->GetName(), *InObj->GetName()));
+		return PyGenUtil::PackReturnValues(FuncParams.GetMemory(), InFuncDef.OutputParams, InErrorCtxt, *FString::Printf(TEXT("function '%s.%s' on '%s'"), *InFuncDef.Func->GetOwnerClass()->GetName(), *InFuncDef.Func->GetName(), *InObj->GetName()));
 	}
 
 	Py_RETURN_NONE;
@@ -357,11 +357,11 @@ PyObject* FPyWrapperObject::CallDynamicFunction_Impl(FPyWrapperObject* InSelf, P
 			}
 		}
 
-		FStructOnScope FuncParams(InFuncDef.Func);
-		PyGenUtil::ApplyParamDefaults(FuncParams.GetStructMemory(), InFuncDef.InputParams);
+		PY_UFUNCTION_STACK(FuncParams, InFuncDef.Func);
+		PyGenUtil::ApplyParamDefaults(FuncParams.GetMemory(), InFuncDef.InputParams);
 		if (ensureAlways(CastField<FObjectPropertyBase>(InSelfParam.ParamProp)))
 		{
-			void* SelfArgInstance = InSelfParam.ParamProp->ContainerPtrToValuePtr<void>(FuncParams.GetStructMemory());
+			void* SelfArgInstance = InSelfParam.ParamProp->ContainerPtrToValuePtr<void>(FuncParams.GetMemory());
 			CastField<FObjectPropertyBase>(InSelfParam.ParamProp)->SetObjectPropertyValue(SelfArgInstance, InSelf->ObjectInstance);
 		}
 		for (int32 ParamIndex = 0; ParamIndex < Params.Num(); ++ParamIndex)
@@ -371,7 +371,7 @@ PyObject* FPyWrapperObject::CallDynamicFunction_Impl(FPyWrapperObject* InSelf, P
 			PyObject* PyValue = Params[ParamIndex];
 			if (PyValue)
 			{
-				if (!PyConversion::NativizeProperty_InContainer(PyValue, ParamDef.ParamProp, FuncParams.GetStructMemory(), 0))
+				if (!PyConversion::NativizeProperty_InContainer(PyValue, ParamDef.ParamProp, FuncParams.GetMemory(), 0))
 				{
 					PyUtil::SetPythonError(PyExc_TypeError, InSelf, *FString::Printf(TEXT("Failed to convert parameter '%s' when calling function '%s.%s' on '%s'"), UTF8_TO_TCHAR(ParamDef.ParamName.GetData()), *Class->GetName(), *InFuncDef.Func->GetName(), *Obj->GetName()));
 					return nullptr;
@@ -379,11 +379,11 @@ PyObject* FPyWrapperObject::CallDynamicFunction_Impl(FPyWrapperObject* InSelf, P
 			}
 		}
 		const FString ErrorCtxt = PyUtil::GetErrorContext(InSelf);
-		if (!PyUtil::InvokeFunctionCall(Obj, InFuncDef.Func, FuncParams.GetStructMemory(), *ErrorCtxt))
+		if (!PyUtil::InvokeFunctionCall(Obj, InFuncDef.Func, FuncParams.GetMemory(), *ErrorCtxt))
 		{
 			return nullptr;
 		}
-		return PyGenUtil::PackReturnValues(FuncParams.GetStructMemory(), InFuncDef.OutputParams, *ErrorCtxt, *FString::Printf(TEXT("function '%s.%s' on '%s'"), *Class->GetName(), *InFuncDef.Func->GetName(), *Obj->GetName()));
+		return PyGenUtil::PackReturnValues(FuncParams.GetMemory(), InFuncDef.OutputParams, *ErrorCtxt, *FString::Printf(TEXT("function '%s.%s' on '%s'"), *Class->GetName(), *InFuncDef.Func->GetName(), *Obj->GetName()));
 	}
 
 	Py_RETURN_NONE;
diff --git a/Engine/Plugins/Experimental/PythonScriptPlugin/Source/PythonScriptPlugin/Private/PyWrapperStruct.cpp b/Engine/Plugins/Experimental/PythonScriptPlugin/Source/PythonScriptPlugin/Private/PyWrapperStruct.cpp
--- a/Engine/Plugins/Experimental/PythonScriptPlugin/Source/PythonScriptPlugin/Private/PyWrapperStruct.cpp
+++ b/Engine/Plugins/Experimental/PythonScriptPlugin/Source/PythonScriptPlugin/Private/PyWrapperStruct.cpp
@@ -467,8 +467,8 @@ int FPyWrapperStruct::CallMakeFunction_Impl(FPyWrapperStruct* InSelf, PyObject*
 		UClass* Class = InFuncDef.Func->GetOwnerClass();
 		UObject* Obj = Class->GetDefaultObject();
 
-		FStructOnScope FuncParams(InFuncDef.Func);
-		PyGenUtil::ApplyParamDefaults(FuncParams.GetStructMemory(), InFuncDef.InputParams);
+		PY_UFUNCTION_STACK(FuncParams, InFuncDef.Func);
+		PyGenUtil::ApplyParamDefaults(FuncParams.GetMemory(), InFuncDef.InputParams);
 		for (int32 ParamIndex = 0; ParamIndex < Params.Num(); ++ParamIndex)
 		{
 			const PyGenUtil::FGeneratedWrappedMethodParameter& ParamDef = InFuncDef.InputParams[ParamIndex];
@@ -476,14 +476,14 @@ int FPyWrapperStruct::CallMakeFunction_Impl(FPyWrapperStruct* InSelf, PyObject*
 			PyObject* PyValue = Params[ParamIndex];
 			if (PyValue)
 			{
-				if (!PyConversion::NativizeProperty_InContainer(PyValue, ParamDef.ParamProp, FuncParams.GetStructMemory(), 0))
+				if (!PyConversion::NativizeProperty_InContainer(PyValue, ParamDef.ParamProp, FuncParams.GetMemory(), 0))
 				{
 					PyUtil::SetPythonError(PyExc_TypeError, InSelf, *FString::Printf(TEXT("Failed to convert parameter '%s' when calling function '%s.%s' on '%s'"), UTF8_TO_TCHAR(ParamDef.ParamName.GetData()), *Class->GetName(), *InFuncDef.Func->GetName(), *Obj->GetName()));
 					return -1;
 				}
 			}
 		}
-		if (!PyUtil::InvokeFunctionCall(Obj, InFuncDef.Func, FuncParams.GetStructMemory(), *PyUtil::GetErrorContext(InSelf)))
+		if (!PyUtil::InvokeFunctionCall(Obj, InFuncDef.Func, FuncParams.GetMemory(), *PyUtil::GetErrorContext(InSelf)))
 		{
 			return -1;
 		}
@@ -491,7 +491,7 @@ int FPyWrapperStruct::CallMakeFunction_Impl(FPyWrapperStruct* InSelf, PyObject*
 		{
 			// Copy the result back onto ourself
 			const PyGenUtil::FGeneratedWrappedMethodParameter& ReturnParam = InFuncDef.OutputParams[0];
-			const void* ReturnArgInstance = ReturnParam.ParamProp->ContainerPtrToValuePtr<void>(FuncParams.GetStructMemory());
+			const void* ReturnArgInstance = ReturnParam.ParamProp->ContainerPtrToValuePtr<void>(FuncParams.GetMemory());
 			InSelf->ScriptStruct->CopyScriptStruct(InSelf->StructInstance, ReturnArgInstance);
 		}
 	}
@@ -506,15 +506,15 @@ PyObject* FPyWrapperStruct::CallBreakFunction_Impl(FPyWrapperStruct* InSelf, con
 		UClass* Class = InFuncDef.Func->GetOwnerClass();
 		UObject* Obj = Class->GetDefaultObject();
 
-		FStructOnScope FuncParams(InFuncDef.Func);
+		PY_UFUNCTION_STACK(FuncParams, InFuncDef.Func);
 		if (ensureAlways(InFuncDef.InputParams.Num() == 1 && CastField<FStructProperty>(InFuncDef.InputParams[0].ParamProp) && InSelf->ScriptStruct->IsChildOf(CastFieldChecked<const FStructProperty>(InFuncDef.InputParams[0].ParamProp)->Struct)))
 		{
 			// Copy us as the 'self' argument
 			const PyGenUtil::FGeneratedWrappedMethodParameter& SelfParam = InFuncDef.InputParams[0];
-			void* SelfArgInstance = SelfParam.ParamProp->ContainerPtrToValuePtr<void>(FuncParams.GetStructMemory());
+			void* SelfArgInstance = SelfParam.ParamProp->ContainerPtrToValuePtr<void>(FuncParams.GetMemory());
 			CastFieldChecked<const FStructProperty>(SelfParam.ParamProp)->Struct->CopyScriptStruct(SelfArgInstance, InSelf->StructInstance);
 		}
-		if (!PyUtil::InvokeFunctionCall(Obj, InFuncDef.Func, FuncParams.GetStructMemory(), *PyUtil::GetErrorContext(InSelf)))
+		if (!PyUtil::InvokeFunctionCall(Obj, InFuncDef.Func, FuncParams.GetMemory(), *PyUtil::GetErrorContext(InSelf)))
 		{
 			return nullptr;
 		}
@@ -524,7 +524,7 @@ PyObject* FPyWrapperStruct::CallBreakFunction_Impl(FPyWrapperStruct* InSelf, con
 			const PyGenUtil::FGeneratedWrappedMethodParameter& ParamDef = InFuncDef.OutputParams[ParamIndex];
 
 			PyObject* PyValue = nullptr;
-			if (!PyConversion::PythonizeProperty_InContainer(ParamDef.ParamProp, FuncParams.GetStructMemory(), 0, PyValue, EPyConversionMethod::Steal))
+			if (!PyConversion::PythonizeProperty_InContainer(ParamDef.ParamProp, FuncParams.GetMemory(), 0, PyValue, EPyConversionMethod::Steal))
 			{
 				PyUtil::SetPythonError(PyExc_TypeError, InSelf, *FString::Printf(TEXT("Failed to convert return property '%s' when calling function '%s.%s' on '%s'"), UTF8_TO_TCHAR(ParamDef.ParamName.GetData()), *Class->GetName(), *InFuncDef.Func->GetName(), *Obj->GetName()));
 				return nullptr;
@@ -560,11 +560,11 @@ PyObject* FPyWrapperStruct::CallDynamicFunction_Impl(FPyWrapperStruct* InSelf, P
 			}
 		}
 
-		FStructOnScope FuncParams(InFuncDef.Func);
-		PyGenUtil::ApplyParamDefaults(FuncParams.GetStructMemory(), InFuncDef.InputParams);
+		PY_UFUNCTION_STACK(FuncParams, InFuncDef.Func);
+		PyGenUtil::ApplyParamDefaults(FuncParams.GetMemory(), InFuncDef.InputParams);
 		if (ensureAlways(CastField<FStructProperty>(InSelfParam.ParamProp) && InSelf->ScriptStruct->IsChildOf(CastFieldChecked<const FStructProperty>(InSelfParam.ParamProp)->Struct)))
 		{
-			void* SelfArgInstance = InSelfParam.ParamProp->ContainerPtrToValuePtr<void>(FuncParams.GetStructMemory());
+			void* SelfArgInstance = InSelfParam.ParamProp->ContainerPtrToValuePtr<void>(FuncParams.GetMemory());
 			CastFieldChecked<const FStructProperty>(InSelfParam.ParamProp)->Struct->CopyScriptStruct(SelfArgInstance, InSelf->StructInstance);
 		}
 		for (int32 ParamIndex = 0; ParamIndex < Params.Num(); ++ParamIndex)
@@ -574,7 +574,7 @@ PyObject* FPyWrapperStruct::CallDynamicFunction_Impl(FPyWrapperStruct* InSelf, P
 			PyObject* PyValue = Params[ParamIndex];
 			if (PyValue)
 			{
-				if (!PyConversion::NativizeProperty_InContainer(PyValue, ParamDef.ParamProp, FuncParams.GetStructMemory(), 0))
+				if (!PyConversion::NativizeProperty_InContainer(PyValue, ParamDef.ParamProp, FuncParams.GetMemory(), 0))
 				{
 					PyUtil::SetPythonError(PyExc_TypeError, InSelf, *FString::Printf(TEXT("Failed to convert parameter '%s' when calling function '%s.%s' on '%s'"), UTF8_TO_TCHAR(ParamDef.ParamName.GetData()), *Class->GetName(), *InFuncDef.Func->GetName(), *Obj->GetName()));
 					return nullptr;
@@ -582,17 +582,17 @@ PyObject* FPyWrapperStruct::CallDynamicFunction_Impl(FPyWrapperStruct* InSelf, P
 			}
 		}
 		const FString ErrorCtxt = PyUtil::GetErrorContext(InSelf);
-		if (!PyUtil::InvokeFunctionCall(Obj, InFuncDef.Func, FuncParams.GetStructMemory(), *ErrorCtxt))
+		if (!PyUtil::InvokeFunctionCall(Obj, InFuncDef.Func, FuncParams.GetMemory(), *ErrorCtxt))
 		{
 			return nullptr;
 		}
 		if (InSelfReturn.ParamProp && ensureAlways(CastField<FStructProperty>(InSelfReturn.ParamProp) && CastFieldChecked<const FStructProperty>(InSelfReturn.ParamProp)->Struct->IsChildOf(InSelf->ScriptStruct)))
 		{
 			// Copy the 'self' return value back onto ourself
-			const void* SelfReturnInstance = InSelfReturn.ParamProp->ContainerPtrToValuePtr<void>(FuncParams.GetStructMemory());
+			const void* SelfReturnInstance = InSelfReturn.ParamProp->ContainerPtrToValuePtr<void>(FuncParams.GetMemory());
 			InSelf->ScriptStruct->CopyScriptStruct(InSelf->StructInstance, SelfReturnInstance);
 		}
-		return PyGenUtil::PackReturnValues(FuncParams.GetStructMemory(), InFuncDef.OutputParams, *ErrorCtxt, *FString::Printf(TEXT("function '%s.%s' on '%s'"), *Class->GetName(), *InFuncDef.Func->GetName(), *Obj->GetName()));
+		return PyGenUtil::PackReturnValues(FuncParams.GetMemory(), InFuncDef.OutputParams, *ErrorCtxt, *FString::Printf(TEXT("function '%s.%s' on '%s'"), *Class->GetName(), *InFuncDef.Func->GetName(), *Obj->GetName()));
 	}
 
 	Py_RETURN_NONE;
@@ -630,11 +630,11 @@ PyObject* FPyWrapperStruct::CallOperatorFunction_Impl(FPyWrapperStruct* InSelf,
 		UObject* Obj = Class->GetDefaultObject();
 
 		// Build the input arguments (failures here aren't fatal as we may have multiple functions to evaluate on the stack, only one of which may accept the RHS parameter)
-		FStructOnScope FuncParams(InOpFunc.Func);
-		PyGenUtil::ApplyParamDefaults(FuncParams.GetStructMemory(), InOpFunc.AdditionalParams);
+		PY_UFUNCTION_STACK(FuncParams, InOpFunc.Func);
+		PyGenUtil::ApplyParamDefaults(FuncParams.GetMemory(), InOpFunc.AdditionalParams);
 		if (InOpFunc.OtherParam.ParamProp)
 		{
-			const FPyConversionResult RHSResult = PyConversion::NativizeProperty_InContainer(InRHS, InOpFunc.OtherParam.ParamProp, FuncParams.GetStructMemory(), 0, nullptr, PyConversion::ESetErrorState::No);
+			const FPyConversionResult RHSResult = PyConversion::NativizeProperty_InContainer(InRHS, InOpFunc.OtherParam.ParamProp, FuncParams.GetMemory(), 0, nullptr, PyConversion::ESetErrorState::No);
 			SetOptionalPyConversionResult(RHSResult, OutRHSConversionResult);
 
 			if (!RHSResult)
@@ -649,10 +649,10 @@ PyObject* FPyWrapperStruct::CallOperatorFunction_Impl(FPyWrapperStruct* InSelf,
 		}
 		if (ensureAlways(CastField<FStructProperty>(InOpFunc.SelfParam.ParamProp) && InSelf->ScriptStruct->IsChildOf(CastFieldChecked<const FStructProperty>(InOpFunc.SelfParam.ParamProp)->Struct)))
 		{
-			void* StructArgInstance = InOpFunc.SelfParam.ParamProp->ContainerPtrToValuePtr<void>(FuncParams.GetStructMemory());
+			void* StructArgInstance = InOpFunc.SelfParam.ParamProp->ContainerPtrToValuePtr<void>(FuncParams.GetMemory());
 			CastFieldChecked<const FStructProperty>(InOpFunc.SelfParam.ParamProp)->Struct->CopyScriptStruct(StructArgInstance, InSelf->StructInstance);
 		}
-		if (!PyUtil::InvokeFunctionCall(Obj, InOpFunc.Func, FuncParams.GetStructMemory(), *PyUtil::GetErrorContext(InSelf)))
+		if (!PyUtil::InvokeFunctionCall(Obj, InOpFunc.Func, FuncParams.GetMemory(), *PyUtil::GetErrorContext(InSelf)))
 		{
 			return nullptr;
 		}
@@ -663,7 +663,7 @@ PyObject* FPyWrapperStruct::CallOperatorFunction_Impl(FPyWrapperStruct* InSelf,
 			if (ensureAlways(CastField<FStructProperty>(InOpFunc.SelfReturn.ParamProp) && CastFieldChecked<const FStructProperty>(InOpFunc.SelfReturn.ParamProp)->Struct->IsChildOf(InSelf->ScriptStruct)))
 			{
 				// Copy the 'self' return value back onto ourself
-				const void* SelfReturnInstance = InOpFunc.SelfReturn.ParamProp->ContainerPtrToValuePtr<void>(FuncParams.GetStructMemory());
+				const void* SelfReturnInstance = InOpFunc.SelfReturn.ParamProp->ContainerPtrToValuePtr<void>(FuncParams.GetMemory());
 				InSelf->ScriptStruct->CopyScriptStruct(InSelf->StructInstance, SelfReturnInstance);
 			}
 
@@ -672,7 +672,7 @@ PyObject* FPyWrapperStruct::CallOperatorFunction_Impl(FPyWrapperStruct* InSelf,
 		}
 		else if (InOpFunc.ReturnParam.ParamProp)
 		{
-			if (!PyConversion::PythonizeProperty_InContainer(InOpFunc.ReturnParam.ParamProp, FuncParams.GetStructMemory(), 0, ReturnPyObj, EPyConversionMethod::Steal))
+			if (!PyConversion::PythonizeProperty_InContainer(InOpFunc.ReturnParam.ParamProp, FuncParams.GetMemory(), 0, ReturnPyObj, EPyConversionMethod::Steal))
 			{
 				PyUtil::SetPythonError(PyExc_TypeError, InSelf, *FString::Printf(TEXT("Failed to convert return property '%s' (%s) when calling function '%s' on '%s'"), *InOpFunc.ReturnParam.ParamProp->GetName(), *InOpFunc.ReturnParam.ParamProp->GetClass()->GetName(), *InOpFunc.Func->GetName(), *Obj->GetName()));
 				return nullptr;
diff --git a/Engine/Plugins/Experimental/PythonScriptPlugin/Source/PythonScriptPlugin/Private/PyWrapperTypeRegistry.cpp b/Engine/Plugins/Experimental/PythonScriptPlugin/Source/PythonScriptPlugin/Private/PyWrapperTypeRegistry.cpp
--- a/Engine/Plugins/Experimental/PythonScriptPlugin/Source/PythonScriptPlugin/Private/PyWrapperTypeRegistry.cpp
+++ b/Engine/Plugins/Experimental/PythonScriptPlugin/Source/PythonScriptPlugin/Private/PyWrapperTypeRegistry.cpp
@@ -3082,8 +3082,8 @@ void FPyWrapperTypeRegistry::GenerateStubCodeForWrappedType(PyTypeObject* PyType
 			}
 			else
 			{
-				FStructOnScope FuncParams(InTypeMethod.MethodFunc.Func);
-				MethodReturnValueStr = GetFunctionReturnValue(FuncParams.GetStructMemory(), InTypeMethod.MethodFunc.OutputParams);
+				PY_UFUNCTION_STACK(FuncParams, InTypeMethod.MethodFunc.Func);
+				MethodReturnValueStr = GetFunctionReturnValue(FuncParams.GetMemory(), InTypeMethod.MethodFunc.OutputParams);
 			}
 		}
 		else
@@ -3120,8 +3120,8 @@ void FPyWrapperTypeRegistry::GenerateStubCodeForWrappedType(PyTypeObject* PyType
 			UClass* Class = InTypeConstant.ConstantFunc.Func->GetOwnerClass();
 			UObject* Obj = Class->GetDefaultObject();
 
-			FStructOnScope FuncParams(InTypeConstant.ConstantFunc.Func);
-			PyUtil::InvokeFunctionCall(Obj, InTypeConstant.ConstantFunc.Func, FuncParams.GetStructMemory(), TEXT("export generated constant"));
+			PY_UFUNCTION_STACK(FuncParams, InTypeConstant.ConstantFunc.Func);
+			PyUtil::InvokeFunctionCall(Obj, InTypeConstant.ConstantFunc.Func, FuncParams.GetMemory(), TEXT("export generated constant"));
 			PyErr_Clear(); // Clear any errors in case InvokeFunctionCall failed
 
 			if (PyGenUtil::IsTypeHintingEnabled())
@@ -3131,7 +3131,7 @@ void FPyWrapperTypeRegistry::GenerateStubCodeForWrappedType(PyTypeObject* PyType
 			}
 			else
 			{
-				ConstantValueStr = GetFunctionReturnValue(FuncParams.GetStructMemory(), InTypeConstant.ConstantFunc.OutputParams);
+				ConstantValueStr = GetFunctionReturnValue(FuncParams.GetMemory(), InTypeConstant.ConstantFunc.OutputParams);
 			}
 		}
 		else // Cannot resove type/value.
