Avoid calling GetStructureSize twice

#rnx

diff --git a/Engine/Plugins/Experimental/PythonScriptPlugin/Source/PythonScriptPlugin/Private/PyUtil.h b/Engine/Plugins/Experimental/PythonScriptPlugin/Source/PythonScriptPlugin/Private/PyUtil.h
--- a/Engine/Plugins/Experimental/PythonScriptPlugin/Source/PythonScriptPlugin/Private/PyUtil.h
+++ b/Engine/Plugins/Experimental/PythonScriptPlugin/Source/PythonScriptPlugin/Private/PyUtil.h
@@ -302,7 +302,7 @@ namespace PyUtil
 		void* StackMemory = nullptr;
 	};
 	#define PY_UFUNCTION_STACK(NAME, FUNCTION) \
-		PyUtil::FFunctionStackOnScope NAME(FUNCTION, FMemory_Alloca_Aligned(((FUNCTION)->GetStructureSize() ? (FUNCTION)->GetStructureSize() : 1), (FUNCTION)->GetMinAlignment()))
+		PyUtil::FFunctionStackOnScope NAME(FUNCTION, FMemory_Alloca_Aligned(FMath::Max(1, (FUNCTION)->GetStructureSize()), (FUNCTION)->GetMinAlignment()))
 
 	/** Struct containing information needed to construct a property instance */
 	struct FPropertyDef
