Clarified the use of args and kwargs in unreal.Object.call_method

#rnx

diff --git a/Engine/Plugins/Experimental/PythonScriptPlugin/Source/PythonScriptPlugin/Private/PyWrapperObject.cpp b/Engine/Plugins/Experimental/PythonScriptPlugin/Source/PythonScriptPlugin/Private/PyWrapperObject.cpp
--- a/Engine/Plugins/Experimental/PythonScriptPlugin/Source/PythonScriptPlugin/Private/PyWrapperObject.cpp
+++ b/Engine/Plugins/Experimental/PythonScriptPlugin/Source/PythonScriptPlugin/Private/PyWrapperObject.cpp
@@ -1102,7 +1102,7 @@ PyTypeObject InitializePyWrapperObjectType()
 		{ "get_editor_property", PyCFunctionCast(&FMethods::GetEditorProperty), METH_VARARGS | METH_KEYWORDS, "get_editor_property(self, name: str) -> object -- get the value of any property visible to the editor" },
 		{ "set_editor_property", PyCFunctionCast(&FMethods::SetEditorProperty), METH_VARARGS | METH_KEYWORDS, "set_editor_property(self, name: str, value: object, notify_mode: PropertyAccessChangeNotifyMode=PropertyAccessChangeNotifyMode.DEFAULT) -> None -- set the value of any property visible to the editor, ensuring that the pre/post change notifications are called" },
 		{ "set_editor_properties", PyCFunctionCast(&FMethods::SetEditorProperties), METH_VARARGS, "set_editor_properties(self, properties: Mapping[str, object]) -> None -- set the value of any properties visible to the editor (from a name->value dict), ensuring that the pre/post change notifications are called" },
-		{ "call_method", PyCFunctionCast(&FMethods::CallMethod), METH_VARARGS | METH_KEYWORDS, "call_method(self, name: str, *args: Any, **kwargs: Mapping[str, object]) -> Any -- call a method on this object via Unreal reflection using the given ordered (tuple) or named (dict) argument data - allows calling methods that don't have Python glue" },
+		{ "call_method", PyCFunctionCast(&FMethods::CallMethod), METH_VARARGS | METH_KEYWORDS, "call_method(self, name: str, *args: Tuple[object, ...], **kwargs: Mapping[str, object]) -> Any -- call a method on this object via Unreal reflection using the given ordered (a tuple passed as args) or named (a dict passed as kwargs) argument data - allows calling methods that don't have Python glue" },
 		{ nullptr, nullptr, 0, nullptr }
 	};
 
