PipInstall: Enable pip installer by default on startup to support broader testing.

diff --git a/Engine/Plugins/Experimental/PythonScriptPlugin/Source/PythonScriptPlugin/Private/PythonScriptPluginSettings.h b/Engine/Plugins/Experimental/PythonScriptPlugin/Source/PythonScriptPlugin/Private/PythonScriptPluginSettings.h
--- a/Engine/Plugins/Experimental/PythonScriptPlugin/Source/PythonScriptPlugin/Private/PythonScriptPluginSettings.h
+++ b/Engine/Plugins/Experimental/PythonScriptPlugin/Source/PythonScriptPlugin/Private/PythonScriptPluginSettings.h
@@ -87,7 +87,7 @@ public:
 	 *       See <ProjectDir>/Intermediate/PipInstall/merged_requirements.in for listing of required packages.
 	 */
 	UPROPERTY(config, EditAnywhere, Category=PythonPipInstall, meta=(ConfigRestartRequired=true))
-	bool bRunPipInstallOnStartup;
+	bool bRunPipInstallOnStartup = true;
 
 	/** Should remote Python execution be enabled? */
 	UPROPERTY(config, EditAnywhere, Category=PythonRemoteExecution, meta=(DisplayName="Enable Remote Execution?"))
