PipInstall: Replicated pip plugin requirements generation functionality directly in engine to increase editor load speed

diff --git a/Engine/Plugins/Experimental/PythonScriptPlugin/Source/PythonScriptPlugin/Private/PipInstall.cpp b/Engine/Plugins/Experimental/PythonScriptPlugin/Source/PythonScriptPlugin/Private/PipInstall.cpp
new file mode 100644
--- /dev/null
+++ b/Engine/Plugins/Experimental/PythonScriptPlugin/Source/PythonScriptPlugin/Private/PipInstall.cpp
@@ -0,0 +1,94 @@
+// Copyright Epic Games, Inc. All Rights Reserved.
+
+#include "PipInstall.h"
+
+#include "Dom/JsonObject.h"
+#include "Misc/FileHelper.h"
+#include "Misc/Paths.h"
+
+
+#if WITH_PYTHON
+
+FString FPipInstall::WritePluginsListing(TArray<TSharedRef<IPlugin>>& OutPythonPlugins)
+{
+	TRACE_CPUPROFILER_EVENT_SCOPE(FPipInstall::WritePluginsListing)
+
+	const FString PipInstallPath = FPaths::ProjectIntermediateDir() / TEXT("PipInstall");
+
+	OutPythonPlugins.Empty();
+
+	TArray<FString> PythonPluginPaths;
+	for ( const TSharedRef<IPlugin>& Plugin : IPluginManager::Get().GetEnabledPlugins() )
+	{
+		const FPluginDescriptor& PluginDesc = Plugin->GetDescriptor();
+		if (PluginDesc.CachedJson->HasTypedField(TEXT("PythonRequirements"), EJson::Array))
+		{
+			const FString PluginDescFile = FPaths::ConvertRelativePathToFull(Plugin->GetDescriptorFileName());
+			PythonPluginPaths.Add(PluginDescFile);
+			OutPythonPlugins.Add(Plugin);
+		}
+	}
+
+	const FString PyPluginsListingFile = PipInstallPath / TEXT("pyreqs_plugins.list");
+	FFileHelper::SaveStringArrayToFile(PythonPluginPaths, *PyPluginsListingFile);
+
+    return PyPluginsListingFile;
+}
+
+FString FPipInstall::WritePluginDependencies(const TArray<TSharedRef<IPlugin>>& PythonPlugins, TArray<FString>& OutRequirements, TArray<FString>& OutExtraUrls)
+{
+	TRACE_CPUPROFILER_EVENT_SCOPE(FPipInstall::WritePluginDependencies)
+
+	const FString PipInstallPath = FPaths::ProjectIntermediateDir() / TEXT("PipInstall");
+
+	OutRequirements.Empty();
+	OutExtraUrls.Empty();
+
+	for (const TSharedRef<IPlugin>& Plugin : PythonPlugins)
+	{
+		const FPluginDescriptor& PluginDesc = Plugin->GetDescriptor();
+		for (const TSharedPtr<FJsonValue>& JsonVal : PluginDesc.CachedJson->GetArrayField(TEXT("PythonRequirements")))
+		{
+			const TSharedPtr<FJsonObject>& JsonObj = JsonVal->AsObject();
+			if (!CheckCompatiblePlatform(JsonObj, FPlatformMisc::GetUBTPlatform()))
+			{
+				continue;
+			}
+
+			const TArray<TSharedPtr<FJsonValue>>* PyReqs;
+			if (JsonObj->TryGetArrayField(TEXT("Requirements"), PyReqs))
+			{
+				for (const TSharedPtr<FJsonValue>& JsonReqVal : *PyReqs)
+				{
+					OutRequirements.Add(JsonReqVal->AsString());
+				}
+			}
+
+			const TArray<TSharedPtr<FJsonValue>>* PyUrls;
+			if (JsonObj->TryGetArrayField(TEXT("ExtraIndexUrls"), PyUrls))
+			{
+				for (const TSharedPtr<FJsonValue>& JsonUrlVal : *PyUrls)
+				{
+					OutExtraUrls.Add(JsonUrlVal->AsString());
+				}
+			}
+		}
+	}
+
+	const FString MergedReqsFile = PipInstallPath / TEXT("merged_requirements.in");
+	const FString ExtraUrlsFile = PipInstallPath / TEXT("extra_urls.txt");
+
+	FFileHelper::SaveStringArrayToFile(OutRequirements, *MergedReqsFile);
+	FFileHelper::SaveStringArrayToFile(OutExtraUrls, *ExtraUrlsFile);
+
+	return MergedReqsFile;
+}
+
+bool FPipInstall::CheckCompatiblePlatform(const TSharedPtr<FJsonObject>& JsonObject, const FString& PlatformName)
+{
+	FString JsonPlatform;
+
+	return !JsonObject->TryGetStringField(TEXT("Platform"), JsonPlatform) || JsonPlatform.Equals(TEXT("All"), ESearchCase::IgnoreCase) || JsonPlatform.Equals(PlatformName, ESearchCase::IgnoreCase);
+}
+
+#endif //WITH_PYTHON
diff --git a/Engine/Plugins/Experimental/PythonScriptPlugin/Source/PythonScriptPlugin/Private/PipInstall.h b/Engine/Plugins/Experimental/PythonScriptPlugin/Source/PythonScriptPlugin/Private/PipInstall.h
new file mode 100644
--- /dev/null
+++ b/Engine/Plugins/Experimental/PythonScriptPlugin/Source/PythonScriptPlugin/Private/PipInstall.h
@@ -0,0 +1,23 @@
+// Copyright Epic Games, Inc. All Rights Reserved.
+
+#pragma once
+
+#include "Containers/Array.h"
+#include "Interfaces/IPluginManager.h"
+#include "Templates/SharedPointer.h"
+
+#if WITH_PYTHON
+
+class FJsonObject;
+
+class FPipInstall
+{
+public:
+	static FString WritePluginsListing(TArray<TSharedRef<IPlugin>>& OutPythonPlugins);
+    static FString WritePluginDependencies(const TArray<TSharedRef<IPlugin>>& PythonPlugins, TArray<FString>& OutRequirements, TArray<FString>& OutExtraUrls);
+
+private:
+	static bool CheckCompatiblePlatform(const TSharedPtr<FJsonObject>& JsonObject, const FString& PlatformName);
+};
+
+#endif //WITH_PYTHON
diff --git a/Engine/Plugins/Experimental/PythonScriptPlugin/Source/PythonScriptPlugin/Private/PythonScriptPlugin.cpp b/Engine/Plugins/Experimental/PythonScriptPlugin/Source/PythonScriptPlugin/Private/PythonScriptPlugin.cpp
--- a/Engine/Plugins/Experimental/PythonScriptPlugin/Source/PythonScriptPlugin/Private/PythonScriptPlugin.cpp
+++ b/Engine/Plugins/Experimental/PythonScriptPlugin/Source/PythonScriptPlugin/Private/PythonScriptPlugin.cpp
@@ -56,6 +56,7 @@
 #include "ContentBrowserFileDataSource.h"
 #include "Toolkits/GlobalEditorCommonCommands.h"
 #include "Misc/FeedbackContext.h"
+#include "PipInstall.h"
 #endif	// WITH_EDITOR
 
 #if PLATFORM_WINDOWS
@@ -1126,6 +1127,8 @@ void FPythonScriptPlugin::ShutdownPython()
 
 void FPythonScriptPlugin::RunPipInstaller()
 {
+	TRACE_CPUPROFILER_EVENT_SCOPE(FPythonScriptPlugin::RunPipInstaller)
+
 	// Run UBT Pip installer for python dependencies (if any)
 	FFeedbackContext* Context = GWarn;
 	FScopedSlowTask PipInstallTask(0, LOCTEXT("PipInstall.RunTasks", "Running Pip Install Tasks..."), true, *Context);
@@ -1133,10 +1136,13 @@ void FPythonScriptPlugin::RunPipInstaller()
 	const FString PipSitePackagePath = FPaths::ConvertRelativePathToFull(GetPipSitePackagesPath());
 
 	// Generate the input listing files of plugins with python dependencies and the listing of all requirements (installed or not)
-	RunUBTPipAction("GenRequirements", LOCTEXT("PipInstall.CheckingDependencies", "Checking Python Dependencies..."), Context);
-	const FString InReqsFile = FPaths::ConvertRelativePathToFull(FPaths::ProjectIntermediateDir() / TEXT("PipInstall") / TEXT("merged_requirements.in"));
+	TArray<TSharedRef<IPlugin>> PythonPlugins;
+	FPipInstall::WritePluginsListing(PythonPlugins);
+
 	TArray<FString> InReqLines;
-	if (!FPaths::FileExists(InReqsFile) || !FFileHelper::LoadFileToStringArray(InReqLines, *InReqsFile) || InReqLines.IsEmpty() )
+	TArray<FString> ExtraUrls;
+	const FString InReqsFile = FPipInstall::WritePluginDependencies(PythonPlugins, InReqLines, ExtraUrls);
+	if (!FPaths::FileExists(InReqsFile) || InReqLines.IsEmpty())
 	{
 		UE_LOG(LogPython, Display, TEXT("No enabled plugins with python dependencies found, skipping"));
 		return;
@@ -1168,6 +1174,8 @@ void FPythonScriptPlugin::RunPipInstaller()
 
 bool FPythonScriptPlugin::RunUBTPipAction(const FString& Action, const FText& Description, FFeedbackContext* Context)
 {
+	TRACE_CPUPROFILER_EVENT_SCOPE(FPythonScriptPlugin::RunUBTPipAction)
+
 	int32 ExitCode;
 	const FString ProjectFileName = IFileManager::Get().ConvertToAbsolutePathForExternalAppForWrite(*FPaths::GetProjectFilePath());
 	const FString Args = FString::Printf(TEXT("%s %s %s -Project=\"%s\" -Mode=PipInstall -PythonInterpreter=\"%s\" -PipAction=%s -Progress")
