PipInstall: Add 'DisablePipInstall' cmd-line argument for build processes to avoid warnings in localization, etc.

diff --git a/Engine/Plugins/Experimental/PythonScriptPlugin/Source/PythonScriptPlugin/Private/PythonScriptPlugin.cpp b/Engine/Plugins/Experimental/PythonScriptPlugin/Source/PythonScriptPlugin/Private/PythonScriptPlugin.cpp
--- a/Engine/Plugins/Experimental/PythonScriptPlugin/Source/PythonScriptPlugin/Private/PythonScriptPlugin.cpp
+++ b/Engine/Plugins/Experimental/PythonScriptPlugin/Source/PythonScriptPlugin/Private/PythonScriptPlugin.cpp
@@ -1127,7 +1127,7 @@ void FPythonScriptPlugin::ShutdownPython()
 
 void FPythonScriptPlugin::RunPipInstaller()
 {
-	TRACE_CPUPROFILER_EVENT_SCOPE(FPythonScriptPlugin::RunPipInstaller)
+	TRACE_CPUPROFILER_EVENT_SCOPE(FPythonScriptPlugin::RunPipInstaller);
 
 	// Run UBT Pip installer for python dependencies (if any)
 	FFeedbackContext* Context = GWarn;
@@ -1149,17 +1149,19 @@ void FPythonScriptPlugin::RunPipInstaller()
 	}
 
 	// Just return immediately with warning if some python dependencies exist and pip install is disabled
-	if (!GetDefault<UPythonScriptPluginSettings>()->bRunPipInstallOnStartup)
+	bool bCmdLineDisable = FParse::Param(FCommandLine::Get(), TEXT("DisablePipInstall"));
+	if (bCmdLineDisable || !GetDefault<UPythonScriptPluginSettings>()->bRunPipInstallOnStartup)
 	{
-		if (!GIsBuildMachine)
+		if (bCmdLineDisable || GIsBuildMachine)
 		{
-			UE_LOG(LogPython, Warning, TEXT("Enabled plugins have python dependencies, enable 'Run Pip Install On Startup' or install manually to: %s"), *PipSitePackagePath);
-			UE_LOG(LogPython, Warning, TEXT("  See package requirements: % s"), *InReqsFile);
+			// Don't warn if disabled on cmd-line or is build process
+			UE_LOG(LogPython, Display, TEXT("Enabled plugins have python dependencies, install manually to: %s"), *PipSitePackagePath);
+			UE_LOG(LogPython, Display, TEXT("  See package requirements: % s"), *InReqsFile);
 		}
 		else
 		{
-			UE_LOG(LogPython, Display, TEXT("Enabled plugins have python dependencies, install manually to: %s"), *PipSitePackagePath);
-			UE_LOG(LogPython, Display, TEXT("  See package requirements: % s"), *InReqsFile);
+			UE_LOG(LogPython, Warning, TEXT("Enabled plugins have python dependencies, enable 'Run Pip Install On Startup' or install manually to: %s"), *PipSitePackagePath);
+			UE_LOG(LogPython, Warning, TEXT("  See package requirements: % s"), *InReqsFile);
 		}
 		return;
 	}
