PipInstall: Added requirements parsing to check for already-installed data and license table viewer

diff --git a/Engine/Plugins/Experimental/PythonScriptPlugin/Source/PythonScriptPlugin/Private/PipInstall.cpp b/Engine/Plugins/Experimental/PythonScriptPlugin/Source/PythonScriptPlugin/Private/PipInstall.cpp
--- a/Engine/Plugins/Experimental/PythonScriptPlugin/Source/PythonScriptPlugin/Private/PipInstall.cpp
+++ b/Engine/Plugins/Experimental/PythonScriptPlugin/Source/PythonScriptPlugin/Private/PipInstall.cpp
@@ -2,18 +2,31 @@
 
 #include "PipInstall.h"
 
+#include "PyUtil.h"
 #include "Dom/JsonObject.h"
+#include "HAL/PlatformFileManager.h"
+#include "Misc/FeedbackContext.h"
+#include "Misc/FeedbackContextMarkup.h"
 #include "Misc/FileHelper.h"
 #include "Misc/Paths.h"
+#include "Misc/ScopedSlowTask.h"
+#include "Misc/ScopeExit.h"
 
+#define LOCTEXT_NAMESPACE "PipInstall"
 
 #if WITH_PYTHON
 
+const FString FPipInstall::PluginsListingFilename = TEXT("pyreqs_plugins.list");
+const FString FPipInstall::RequirementsInputFilename = TEXT("merged_requirements.in");
+const FString FPipInstall::ExtraUrlsFilename = TEXT("extra_urls.txt");
+const FString FPipInstall::ParsedRequirementsFilename = TEXT("merged_requirements.txt");
+
+
 FString FPipInstall::WritePluginsListing(TArray<TSharedRef<IPlugin>>& OutPythonPlugins)
 {
-	TRACE_CPUPROFILER_EVENT_SCOPE(FPipInstall::WritePluginsListing)
+	TRACE_CPUPROFILER_EVENT_SCOPE(FPipInstall::WritePluginsListing);
 
-	const FString PipInstallPath = FPaths::ProjectIntermediateDir() / TEXT("PipInstall");
+	const FString PipInstallPath = GetPipInstallPath();
 
 	OutPythonPlugins.Empty();
 
@@ -29,7 +42,7 @@ FString FPipInstall::WritePluginsListing(TArray<TSharedRef<IPlugin>>& OutPythonP
 		}
 	}
 
-	const FString PyPluginsListingFile = PipInstallPath / TEXT("pyreqs_plugins.list");
+	const FString PyPluginsListingFile = PipInstallPath / PluginsListingFilename;
 	FFileHelper::SaveStringArrayToFile(PythonPluginPaths, *PyPluginsListingFile);
 
     return PyPluginsListingFile;
@@ -37,9 +50,9 @@ FString FPipInstall::WritePluginsListing(TArray<TSharedRef<IPlugin>>& OutPythonP
 
 FString FPipInstall::WritePluginDependencies(const TArray<TSharedRef<IPlugin>>& PythonPlugins, TArray<FString>& OutRequirements, TArray<FString>& OutExtraUrls)
 {
-	TRACE_CPUPROFILER_EVENT_SCOPE(FPipInstall::WritePluginDependencies)
+	TRACE_CPUPROFILER_EVENT_SCOPE(FPipInstall::WritePluginDependencies);
 
-	const FString PipInstallPath = FPaths::ProjectIntermediateDir() / TEXT("PipInstall");
+	const FString PipInstallPath = GetPipInstallPath();
 
 	OutRequirements.Empty();
 	OutExtraUrls.Empty();
@@ -75,8 +88,8 @@ FString FPipInstall::WritePluginDependencies(const TArray<TSharedRef<IPlugin>>&
 		}
 	}
 
-	const FString MergedReqsFile = PipInstallPath / TEXT("merged_requirements.in");
-	const FString ExtraUrlsFile = PipInstallPath / TEXT("extra_urls.txt");
+	const FString MergedReqsFile = PipInstallPath / RequirementsInputFilename;
+	const FString ExtraUrlsFile = PipInstallPath / ExtraUrlsFilename;
 
 	FFileHelper::SaveStringArrayToFile(OutRequirements, *MergedReqsFile);
 	FFileHelper::SaveStringArrayToFile(OutExtraUrls, *ExtraUrlsFile);
@@ -84,6 +97,171 @@ FString FPipInstall::WritePluginDependencies(const TArray<TSharedRef<IPlugin>>&
 	return MergedReqsFile;
 }
 
+void FPipInstall::SetupPipEnv(FFeedbackContext* Context, bool bForceRebuild /* = false */)
+{
+	TRACE_CPUPROFILER_EVENT_SCOPE(FPipInstall::SetupPipEnv);
+
+	const FString PipInstallPath = GetPipInstallPath();
+	const FString VenvInterp = GetVenvInterpreter(PipInstallPath);
+
+	if (!bForceRebuild && FPaths::FileExists(VenvInterp))
+	{
+		SetupPipInstallUtils(VenvInterp, Context);
+		return;
+	}
+
+	if (bForceRebuild && FPaths::DirectoryExists(PipInstallPath))
+	{
+		// TODO: Need to cache generated files before deleting the directory (or only delete subdirs)
+		IPlatformFile& PlatformFile = FPlatformFileManager::Get().GetPlatformFile();
+		PlatformFile.DeleteDirectoryRecursively(*PipInstallPath);
+	}
+
+	const FString EngineInterp = PyUtil::GetInterpreterExecutablePath();
+	const FString VenvCmd = FString::Printf(TEXT("-m venv \"%s\""), *FPaths::ConvertRelativePathToFull(PipInstallPath));
+	int32 Res = RunPythonCmd(LOCTEXT("PipInstall.SetupVenv", "Setting up pip install environment..."), EngineInterp, VenvCmd, Context);
+	if (Res != 0)
+	{
+		UE_LOG(LogPython, Error, TEXT("Unable to create pip install environment (%d)"), Res);
+		return;
+	}
+
+	SetupPipInstallUtils(VenvInterp, Context);
+}
+
+FString FPipInstall::ParsePluginDependencies(const FString& MergedInRequirementsFile, FFeedbackContext* Context)
+{
+	TRACE_CPUPROFILER_EVENT_SCOPE(FPipInstall::ParsePluginDependencies);
+
+	const FString PipInstallPath = GetPipInstallPath();
+	const FString VenvInterp = GetVenvInterpreter(PipInstallPath);
+
+	const FString ParsedReqsFile = PipInstallPath / ParsedRequirementsFilename;
+
+	const FString Cmd = FString::Printf(TEXT("-m ue_parse_plugin_reqs -vv \"%s\" \"%s\""), *MergedInRequirementsFile, *ParsedReqsFile);
+	RunPythonCmd(LOCTEXT("PipInstall.ParseRequirements", "Parsing pip requirements..."), VenvInterp, Cmd, Context);
+
+	return FPaths::ConvertRelativePathToFull(ParsedReqsFile);
+}
+
+
+void FPipInstall::SetupPipInstallUtils(const FString& VenvInterp, FFeedbackContext* Context)
+{
+	TRACE_CPUPROFILER_EVENT_SCOPE(FPipInstall::SetupPipInstallUtils);
+
+	if (CheckPipInstallUtils(VenvInterp, Context))
+	{
+		return;
+	}
+
+	const FString PipInstallPath = GetPipInstallPath();
+	const FString PythonScriptDir = GetPythonScriptPluginPath();
+	if (PythonScriptDir.IsEmpty())
+	{
+		return;
+	}
+
+	const FString PipWheelsDir = PythonScriptDir / TEXT("Content/Python/Lib/wheels");
+	const FString InstallRequirements = PythonScriptDir / TEXT("Content/Python/PipInstallUtils/requirements.txt");
+
+	const FString Cmd = FString::Printf(TEXT("-m pip install --upgrade --no-index --find-links \"%s\" -r \"%s\" ue-pipinstall-utils"), *PipWheelsDir, *InstallRequirements);
+
+	RunPythonCmd(LOCTEXT("PipInstall.SetupPipInstallUtils", "Setting up pip install utils"), VenvInterp, Cmd, Context);
+}
+
+
+bool FPipInstall::CheckPipInstallUtils(const FString& VenvInterp, FFeedbackContext* Context)
+{
+	const FString Cmd = TEXT("-c \"import pkg_resources; exit(0 if pkg_resources.working_set.find(pkg_resources.Requirement.parse('ue-pipinstall-utils')) is not None else 1)\"");
+	return (RunPythonCmd(LOCTEXT("PipInstall.CheckPipInstallUtils", "Check pip install utils installed"), VenvInterp, Cmd, Context) == 0);
+}
+
+int32 FPipInstall::RunPythonCmd(const FText& Description, const FString& PythonInterp, const FString& Cmd, FFeedbackContext* Context)
+{
+	TRACE_CPUPROFILER_EVENT_SCOPE(FPipInstall::RunPythonCmd);
+
+	UE_LOG(LogPython, Log, TEXT("Running python command: python %s"), *Cmd);
+
+	int32 Result = 0;
+	RunLoggedSubprocess(Description, FPaths::ConvertRelativePathToFull(PythonInterp), Cmd, Context, &Result);
+
+	return Result;
+}
+
+bool FPipInstall::RunLoggedSubprocess(const FText& Description, const FString& URL, const FString& Params, FFeedbackContext* Context, int32* OutExitCode)
+{
+	FScopedSlowTask SubprocessTask(0, Description, true, *Context);
+
+	// Create a read and write pipe for the child process
+	void* StdOutPipeRead = nullptr;
+	void* StdOutPipeWrite = nullptr;
+	verify(FPlatformProcess::CreatePipe(StdOutPipeRead, StdOutPipeWrite));
+
+	ON_SCOPE_EXIT
+	{
+		FPlatformProcess::ClosePipe(StdOutPipeRead, StdOutPipeWrite);
+	};
+
+	// Create the process
+	FProcHandle ProcessHandle = FPlatformProcess::CreateProc(*URL, *Params, false, true, true, nullptr, 0, nullptr, StdOutPipeWrite, nullptr);
+	if (ProcessHandle.IsValid())
+	{
+		FString BufferedText;
+		for (bool bProcessFinished = false; !bProcessFinished; )
+		{
+			bProcessFinished = FPlatformProcess::GetProcReturnCode(ProcessHandle, OutExitCode);
+			BufferedText += FPlatformProcess::ReadPipe(StdOutPipeRead);
+
+			int32 EndOfLineIdx;
+			while (BufferedText.FindChar(TEXT('\n'), EndOfLineIdx))
+			{
+				FString Line = BufferedText.Left(EndOfLineIdx);
+				Line.RemoveFromEnd(TEXT("\r"), ESearchCase::CaseSensitive);
+
+				Context->Log(LogPython.GetCategoryName(), ELogVerbosity::Log, Line);
+				BufferedText.MidInline(EndOfLineIdx + 1, MAX_int32, false);
+			}
+
+			FPlatformProcess::Sleep(0.1f);
+		}
+		ProcessHandle.Reset();
+		return true;
+	}
+	else
+	{
+		Context->CategorizedLogf(LogPython.GetCategoryName(), ELogVerbosity::Warning, TEXT("Couldn't create process '%s'"), *URL);
+		return false;
+	}
+}
+
+
+FString FPipInstall::GetPipInstallPath()
+{
+	return FPaths::ProjectIntermediateDir() / TEXT("PipInstall");
+}
+
+FString FPipInstall::GetPythonScriptPluginPath()
+{
+	TSharedPtr<IPlugin> PythonPlugin = IPluginManager::Get().FindPlugin("PythonScriptPlugin");
+	if (!PythonPlugin)
+	{
+		return TEXT("");
+	}
+
+	return PythonPlugin->GetBaseDir();
+}
+
+FString FPipInstall::GetVenvInterpreter(const FString& InstallPath)
+{
+#if PLATFORM_WINDOWS
+	return InstallPath / TEXT("Scripts/python.exe");
+#elif PLATFORM_MAC || PLATFORM_LINUX
+	return InstallPath / TEXT("bin/python3");
+#else
+	static_assert(false, "Python not supported on this platform!");
+#endif
+}
+
 bool FPipInstall::CheckCompatiblePlatform(const TSharedPtr<FJsonObject>& JsonObject, const FString& PlatformName)
 {
 	FString JsonPlatform;
@@ -92,3 +270,5 @@ bool FPipInstall::CheckCompatiblePlatform(const TSharedPtr<FJsonObject>& JsonObj
 }
 
 #endif //WITH_PYTHON
+
+#undef LOCTEXT_NAMESPACE
diff --git a/Engine/Plugins/Experimental/PythonScriptPlugin/Source/PythonScriptPlugin/Private/PipInstall.h b/Engine/Plugins/Experimental/PythonScriptPlugin/Source/PythonScriptPlugin/Private/PipInstall.h
--- a/Engine/Plugins/Experimental/PythonScriptPlugin/Source/PythonScriptPlugin/Private/PipInstall.h
+++ b/Engine/Plugins/Experimental/PythonScriptPlugin/Source/PythonScriptPlugin/Private/PipInstall.h
@@ -9,14 +9,32 @@
 #if WITH_PYTHON
 
 class FJsonObject;
+class FFeedbackContext;
 
 class FPipInstall
 {
 public:
+	static const FString PluginsListingFilename;
+	static const FString RequirementsInputFilename;
+	static const FString ExtraUrlsFilename;
+	static const FString ParsedRequirementsFilename;
+
 	static FString WritePluginsListing(TArray<TSharedRef<IPlugin>>& OutPythonPlugins);
-    static FString WritePluginDependencies(const TArray<TSharedRef<IPlugin>>& PythonPlugins, TArray<FString>& OutRequirements, TArray<FString>& OutExtraUrls);
+	static FString WritePluginDependencies(const TArray<TSharedRef<IPlugin>>& PythonPlugins, TArray<FString>& OutRequirements, TArray<FString>& OutExtraUrls);
+
+	static void SetupPipEnv(FFeedbackContext* Context, bool bForceRebuild = false);
+	static FString ParsePluginDependencies(const FString& MergedInRequirementsFile, FFeedbackContext* Context);
+
+	static FString GetPipInstallPath();
 
 private:
+	static void SetupPipInstallUtils(const FString& VenvInterp, FFeedbackContext* Context);
+	static bool CheckPipInstallUtils(const FString& VenvInterp, FFeedbackContext* Context);
+	static int32 RunPythonCmd(const FText& Description, const FString& VenvInterp, const FString& Cmd, FFeedbackContext* Context);
+	static bool RunLoggedSubprocess(const FText& Description, const FString& URL, const FString& Params, FFeedbackContext* Context, int32* OutExitCode);
+
+	static FString GetPythonScriptPluginPath();
+	static FString GetVenvInterpreter(const FString& InstallPath);
 	static bool CheckCompatiblePlatform(const TSharedPtr<FJsonObject>& JsonObject, const FString& PlatformName);
 };
 
diff --git a/Engine/Plugins/Experimental/PythonScriptPlugin/Source/PythonScriptPlugin/Private/PythonScriptPlugin.cpp b/Engine/Plugins/Experimental/PythonScriptPlugin/Source/PythonScriptPlugin/Private/PythonScriptPlugin.cpp
--- a/Engine/Plugins/Experimental/PythonScriptPlugin/Source/PythonScriptPlugin/Private/PythonScriptPlugin.cpp
+++ b/Engine/Plugins/Experimental/PythonScriptPlugin/Source/PythonScriptPlugin/Private/PythonScriptPlugin.cpp
@@ -978,7 +978,8 @@ void FPythonScriptPlugin::InitializePython()
 		FPyWrapperTypeRegistry::Get().GenerateWrappedTypes();
 
 #if WITH_EDITOR
-		// Run PipInstall UBT task
+		// Init PipInstall task
+		InitPipInstaller();
 		RunPipInstaller();
 #endif // WITH_EDITOR
 
@@ -1125,13 +1126,13 @@ void FPythonScriptPlugin::ShutdownPython()
 	bRanStartupScripts = false;
 }
 
-void FPythonScriptPlugin::RunPipInstaller()
+void FPythonScriptPlugin::InitPipInstaller()
 {
-	TRACE_CPUPROFILER_EVENT_SCOPE(FPythonScriptPlugin::RunPipInstaller);
+	TRACE_CPUPROFILER_EVENT_SCOPE(FPythonScriptPlugin::InitPipInstaller);
 
-	// Run UBT Pip installer for python dependencies (if any)
+	// Init Pip installer for python dependencies (if any)
 	FFeedbackContext* Context = GWarn;
-	FScopedSlowTask PipInstallTask(0, LOCTEXT("PipInstall.RunTasks", "Running Pip Install Tasks..."), true, *Context);
+	FScopedSlowTask PipInstallTask(0, LOCTEXT("PipInstall.RunInit", "Running Pip Init Tasks..."), true, *Context);
 
 	const FString PipSitePackagePath = FPaths::ConvertRelativePathToFull(GetPipSitePackagesPath());
 
@@ -1139,35 +1140,64 @@ void FPythonScriptPlugin::RunPipInstaller()
 	TArray<TSharedRef<IPlugin>> PythonPlugins;
 	FPipInstall::WritePluginsListing(PythonPlugins);
 
-	TArray<FString> InReqLines;
+	TArray<FString> ReqInLines;
 	TArray<FString> ExtraUrls;
-	const FString InReqsFile = FPipInstall::WritePluginDependencies(PythonPlugins, InReqLines, ExtraUrls);
-	if (!FPaths::FileExists(InReqsFile) || InReqLines.IsEmpty())
+	const FString ReqsInFile = FPipInstall::WritePluginDependencies(PythonPlugins, ReqInLines, ExtraUrls);
+
+	bool bRunOnStartup = GetDefault<UPythonScriptPluginSettings>()->bRunPipInstallOnStartup;
+	if (ReqInLines.IsEmpty())
 	{
-		UE_LOG(LogPython, Display, TEXT("No enabled plugins with python dependencies found, skipping"));
+		UE_CLOG(bRunOnStartup, LogPython, Display, TEXT("No enabled plugins with python dependencies found, skipping"));
 		return;
 	}
 
 	// Just return immediately with warning if some python dependencies exist and pip install is disabled
 	bool bCmdLineDisable = FParse::Param(FCommandLine::Get(), TEXT("DisablePipInstall"));
-	if (bCmdLineDisable || !GetDefault<UPythonScriptPluginSettings>()->bRunPipInstallOnStartup)
+	if (bCmdLineDisable || !bRunOnStartup)
 	{
 		if (bCmdLineDisable || GIsBuildMachine)
 		{
 			// Don't warn if disabled on cmd-line or is build process
 			UE_LOG(LogPython, Display, TEXT("Enabled plugins have python dependencies, install manually to: %s"), *PipSitePackagePath);
-			UE_LOG(LogPython, Display, TEXT("  See package requirements: % s"), *InReqsFile);
+			UE_LOG(LogPython, Display, TEXT("  See package requirements: % s"), *ReqsInFile);
 		}
 		else
 		{
 			UE_LOG(LogPython, Warning, TEXT("Enabled plugins have python dependencies, enable 'Run Pip Install On Startup' or install manually to: %s"), *PipSitePackagePath);
-			UE_LOG(LogPython, Warning, TEXT("  See package requirements: % s"), *InReqsFile);
+			UE_LOG(LogPython, Warning, TEXT("  See package requirements: % s"), *ReqsInFile);
 		}
 		return;
 	}
 
+	FPipInstall::SetupPipEnv(Context);
+	FPipInstall::ParsePluginDependencies(ReqsInFile, Context);
+}
+
+void FPythonScriptPlugin::RunPipInstaller()
+{
+	TRACE_CPUPROFILER_EVENT_SCOPE(FPythonScriptPlugin::RunPipInstaller);
+
+	bool bRunOnStartup = GetDefault<UPythonScriptPluginSettings>()->bRunPipInstallOnStartup;
+	bool bCmdLineDisable = FParse::Param(FCommandLine::Get(), TEXT("DisablePipInstall"));
+	if (bCmdLineDisable || !bRunOnStartup)
+	{
+		return;
+	}
+
+	const FString ParsedReqsFile = FPaths::ConvertRelativePathToFull(FPipInstall::GetPipInstallPath() / FPipInstall::ParsedRequirementsFilename);
+
+	TArray<FString> ParsedReqLines;
+	if (!FPaths::FileExists(ParsedReqsFile) || !FFileHelper::LoadFileToStringArray(ParsedReqLines, *ParsedReqsFile) || ParsedReqLines.IsEmpty())
+	{
+		return;
+	}
+
+	FFeedbackContext* Context = GWarn;
+	FScopedSlowTask PipInstallTask(0, LOCTEXT("PipInstall.RunInstall", "Installing Python Dependencies..."), true, *Context);
+	//PipInstallTask.MakeDialog(true);
+
 	// Run install of all python dependencies for enabled plugins
-	if (!RunUBTPipAction("InstallNoRegen", LOCTEXT("PipInstall.InstallingDependencies", "Installing Python Dependencies..."), Context))
+	if (!RunUBTPipAction("InstallNoRegen", LOCTEXT("PipInstall.UBTInstall", "Pip Installing Plugin Dependencies..."), Context))
 	{
 		UE_LOG(LogPython, Warning, TEXT("Unable to install plugin python dependencies"));
 		return;
@@ -1176,7 +1206,7 @@ void FPythonScriptPlugin::RunPipInstaller()
 
 bool FPythonScriptPlugin::RunUBTPipAction(const FString& Action, const FText& Description, FFeedbackContext* Context)
 {
-	TRACE_CPUPROFILER_EVENT_SCOPE(FPythonScriptPlugin::RunUBTPipAction)
+	TRACE_CPUPROFILER_EVENT_SCOPE(FPythonScriptPlugin::RunUBTPipAction);
 
 	int32 ExitCode;
 	const FString ProjectFileName = IFileManager::Get().ConvertToAbsolutePathForExternalAppForWrite(*FPaths::GetProjectFilePath());
@@ -1193,7 +1223,7 @@ bool FPythonScriptPlugin::RunUBTPipAction(const FString& Action, const FText& De
 
 FString FPythonScriptPlugin::GetPipSitePackagesPath()
 {
-	const FString VenvPath = FPaths::ProjectIntermediateDir() / TEXT("PipInstall");
+	const FString VenvPath = FPipInstall::GetPipInstallPath();
 #if PLATFORM_WINDOWS
 	return VenvPath / TEXT("Lib") / TEXT("site-packages");
 #elif PLATFORM_MAC || PLATFORM_LINUX
diff --git a/Engine/Plugins/Experimental/PythonScriptPlugin/Source/PythonScriptPlugin/Private/PythonScriptPlugin.h b/Engine/Plugins/Experimental/PythonScriptPlugin/Source/PythonScriptPlugin/Private/PythonScriptPlugin.h
--- a/Engine/Plugins/Experimental/PythonScriptPlugin/Source/PythonScriptPlugin/Private/PythonScriptPlugin.h
+++ b/Engine/Plugins/Experimental/PythonScriptPlugin/Source/PythonScriptPlugin/Private/PythonScriptPlugin.h
@@ -162,6 +162,8 @@ private:
 
 	void ShutdownPython();
 
+	void InitPipInstaller();
+
 	void RunPipInstaller();
 
 	bool RunUBTPipAction(const FString& Action, const FText& Description, FFeedbackContext* Context);
diff --git a/Engine/Source/Programs/UnrealBuildTool/Modes/PipInstallMode.cs b/Engine/Source/Programs/UnrealBuildTool/Modes/PipInstallMode.cs
--- a/Engine/Source/Programs/UnrealBuildTool/Modes/PipInstallMode.cs
+++ b/Engine/Source/Programs/UnrealBuildTool/Modes/PipInstallMode.cs
@@ -144,30 +144,30 @@ namespace UnrealBuildTool.Modes
 				}
 			}
 
-			// if ((Action & (PipAction)ActionBits.ParseReqs) != 0)
-			// {
+			if ((Action & (PipAction)ActionBits.ParseReqs) != 0)
+			{
 
-			// 	if (!Pip.ParsePluginDependencies(MergedReqsInFile, MergedRequirementsFile))
-			// 	{
-			// 		return 1;
-			// 	}
-			// }
+				if (!Pip.ParsePluginDependencies(MergedReqsInFile, MergedRequirementsFile))
+				{
+					return 1;
+				}
+			}
 
 			if ((Action & (PipAction)ActionBits.InstallReqs) != 0)
 			{
-				if (!Pip.InstallPluginDependencies(MergedReqsInFile, ExtraUrlsFile))
+				if (!Pip.InstallPluginDependencies(MergedRequirementsFile, ExtraUrlsFile))
 				{
 					return 1;
 				}
 			}
 
-			// if ((Action & (PipAction)ActionBits.ViewLicenses) != 0)
-			// {
-			// 	if (!Pip.ViewInstalledLicenses(PluginsListingFile))
-			// 	{
-			// 		return 1;
-			// 	}
-			// }
+			if ((Action & (PipAction)ActionBits.ViewLicenses) != 0)
+			{
+				if (!Pip.ViewInstalledLicenses(PluginsListingFile))
+				{
+					return 1;
+				}
+			}
 
 			return 0;
 		}
@@ -245,7 +245,7 @@ namespace UnrealBuildTool.Modes
 					return false;
 				}
 
-				int result = RunPythonCmd(EnginePythonInterp, $"-m venv {InstallDir}", CmdLogger);
+				int result = RunPythonCmd(EnginePythonInterp, $"-m venv \"{InstallDir}\"", CmdLogger);
 				if (result != 0 || !FileReference.Exists(PythonVenv))
 				{
 					return false;
@@ -271,13 +271,13 @@ namespace UnrealBuildTool.Modes
 			return true;
 		}
 
-		// public bool ParsePluginDependencies(FileReference MergedReqsInFile, FileReference MergedRequirmentsFile, bool NoCheckMerged = false)
-		// {
-		// 	using (IBaseCmdProgressLogger CmdLogger = new PythonCmdLogger(Logger))
-		// 	{
-		// 		return (RunPythonVenv($"-m ue_parse_plugin_reqs -vv {MergedReqsInFile} {MergedRequirmentsFile}", CmdLogger) == 0);
-		// 	}
-		// }
+		public bool ParsePluginDependencies(FileReference MergedReqsInFile, FileReference MergedRequirmentsFile, bool NoCheckMerged = false)
+		{
+			using (IBaseCmdProgressLogger CmdLogger = new PythonCmdLogger(Logger))
+			{
+				return (RunPythonVenv($"-m ue_parse_plugin_reqs -vv \"{MergedReqsInFile}\" \"{MergedRequirmentsFile}\"", CmdLogger) == 0);
+			}
+		}
 
 		public bool InstallPluginDependencies(FileReference MergedRequirementsFile, FileReference ExtraUrlsFile, bool OfflineOnly = false, string? ForceIndexUrl = null)
 		{
@@ -362,7 +362,7 @@ namespace UnrealBuildTool.Modes
 				}
 			}
 
-			Args += " -r " + RequirementsFile.ToString();
+			Args += " -r \"" + RequirementsFile.ToString() + "\"";
 
 			using (IBaseCmdProgressLogger StatusLogger = LoggerFactory.Create("Installing Python Dependencies...", RequirementsCount))
 			{
@@ -432,7 +432,20 @@ namespace UnrealBuildTool.Modes
 
 		private bool SetupPipInstallUtils(IBaseCmdProgressLogger StatusLogger)
 		{
-			return true;
+			Logger.LogInformation("PipInstall: Updating UE PipInstall Utilities");
+			FileReference? PythonScriptPlugin = GetPythonScriptPlugin();
+			if (PythonScriptPlugin == null)
+			{
+				Logger.LogError("PipInstall: Unable to locate engine PythonScriptPlugin");
+				return false;
+			}
+			DirectoryReference PythonScriptDir = DirectoryReference.FromFile(PythonScriptPlugin);
+			DirectoryReference PipInstallUtilsDir = DirectoryReference.Combine(PythonScriptDir, "Content", "Python", "PipInstallUtils");
+			DirectoryReference PipWheelsDir = DirectoryReference.Combine(PythonScriptDir, "Content", "Python", "Lib", "wheels");
+			FileReference RequirementsFile = FileReference.Combine(PipInstallUtilsDir, "requirements.txt");
+			string Args = $"-m pip install --upgrade --no-index --find-links \"{PipWheelsDir}\" -r \"{RequirementsFile}\" ue-pipinstall-utils";
+			int Result = RunPythonVenv(Args, StatusLogger);
+			return (Result == 0);
 		}
 
 		private int RunPythonVenv(string Args, IBaseCmdProgressLogger InCmdLogger)
