Refix C4702 undone by CL# 29950593

diff --git a/Engine/Plugins/Experimental/PythonScriptPlugin/Source/PythonScriptPlugin/Private/PyWrapperMap.cpp b/Engine/Plugins/Experimental/PythonScriptPlugin/Source/PythonScriptPlugin/Private/PyWrapperMap.cpp
--- a/Engine/Plugins/Experimental/PythonScriptPlugin/Source/PythonScriptPlugin/Private/PyWrapperMap.cpp
+++ b/Engine/Plugins/Experimental/PythonScriptPlugin/Source/PythonScriptPlugin/Private/PyWrapperMap.cpp
@@ -1074,8 +1074,6 @@ PyObject* FPyWrapperMap::PopItem(FPyWrapperMap* InSelf)
 		
 		return PyTuple_Pack(2, PyReturnKey.Get(), PyReturnValue.Get());
 	}
-
-	return nullptr;
 }
 
 int FPyWrapperMap::Update(FPyWrapperMap* InSelf, PyObject* InOther)
diff --git a/Engine/Plugins/Experimental/PythonScriptPlugin/Source/PythonScriptPlugin/Private/PyWrapperSet.cpp b/Engine/Plugins/Experimental/PythonScriptPlugin/Source/PythonScriptPlugin/Private/PyWrapperSet.cpp
--- a/Engine/Plugins/Experimental/PythonScriptPlugin/Source/PythonScriptPlugin/Private/PyWrapperSet.cpp
+++ b/Engine/Plugins/Experimental/PythonScriptPlugin/Source/PythonScriptPlugin/Private/PyWrapperSet.cpp
@@ -589,8 +589,6 @@ PyObject* FPyWrapperSet::Pop(FPyWrapperSet* InSelf)
 
 		return PyReturnValue;
 	}
-
-	return nullptr;
 }
 
 int FPyWrapperSet::Clear(FPyWrapperSet* InSelf)
