PipInstall: Disable installer for build machines for now

diff --git a/Engine/Plugins/Experimental/PythonScriptPlugin/Source/PythonScriptPlugin/Private/PythonScriptPlugin.cpp b/Engine/Plugins/Experimental/PythonScriptPlugin/Source/PythonScriptPlugin/Private/PythonScriptPlugin.cpp
--- a/Engine/Plugins/Experimental/PythonScriptPlugin/Source/PythonScriptPlugin/Private/PythonScriptPlugin.cpp
+++ b/Engine/Plugins/Experimental/PythonScriptPlugin/Source/PythonScriptPlugin/Private/PythonScriptPlugin.cpp
@@ -1153,7 +1153,7 @@ void FPythonScriptPlugin::InitPipInstaller()
 
 	// Just return immediately with warning if some python dependencies exist and pip install is disabled
 	bool bCmdLineDisable = FParse::Param(FCommandLine::Get(), TEXT("DisablePipInstall"));
-	if (bCmdLineDisable || !bRunOnStartup)
+	if (bCmdLineDisable || !bRunOnStartup || GIsBuildMachine)
 	{
 		if (bCmdLineDisable || GIsBuildMachine)
 		{
