PipInstall: Cleanup installer enable check and re-enable pip install for build machines

diff --git a/Engine/Plugins/Experimental/PythonScriptPlugin/Source/PythonScriptPlugin/Private/PipInstall.cpp b/Engine/Plugins/Experimental/PythonScriptPlugin/Source/PythonScriptPlugin/Private/PipInstall.cpp
--- a/Engine/Plugins/Experimental/PythonScriptPlugin/Source/PythonScriptPlugin/Private/PipInstall.cpp
+++ b/Engine/Plugins/Experimental/PythonScriptPlugin/Source/PythonScriptPlugin/Private/PipInstall.cpp
@@ -3,8 +3,10 @@
 #include "PipInstall.h"
 
 #include "PyUtil.h"
+#include "PythonScriptPluginSettings.h"
 #include "Dom/JsonObject.h"
 #include "HAL/PlatformFileManager.h"
+#include "Misc/CommandLine.h"
 #include "Misc/FeedbackContext.h"
 #include "Misc/FeedbackContextMarkup.h"
 #include "Misc/FileHelper.h"
@@ -27,6 +29,14 @@ const FString FPipInstall::ExtraUrlsFilename = TEXT("extra_urls.txt");
 const FString FPipInstall::ParsedRequirementsFilename = TEXT("merged_requirements.txt");
 
 
+bool FPipInstall::EnabledOnStartup()
+{
+	bool bRunOnStartup = GetDefault<UPythonScriptPluginSettings>()->bRunPipInstallOnStartup;
+	bool bCmdLineDisable = FParse::Param(FCommandLine::Get(), TEXT("DisablePipInstall"));
+
+	return bRunOnStartup && !bCmdLineDisable;
+}
+
 FString FPipInstall::WritePluginsListing(TArray<TSharedRef<IPlugin>>& OutPythonPlugins)
 {
 	TRACE_CPUPROFILER_EVENT_SCOPE(FPipInstall::WritePluginsListing);
@@ -116,8 +126,8 @@ FString FPipInstall::WritePluginDependencies(const TArray<TSharedRef<IPlugin>>&
 		}
 	}
 
-	const FString MergedReqsFile = PipInstallPath / RequirementsInputFilename;
-	const FString ExtraUrlsFile = PipInstallPath / ExtraUrlsFilename;
+	const FString MergedReqsFile = FPaths::ConvertRelativePathToFull(PipInstallPath / RequirementsInputFilename);
+	const FString ExtraUrlsFile = FPaths::ConvertRelativePathToFull(PipInstallPath / ExtraUrlsFilename);
 
 	FFileHelper::SaveStringArrayToFile(OutRequirements, *MergedReqsFile);
 	FFileHelper::SaveStringArrayToFile(OutExtraUrls, *ExtraUrlsFile);
diff --git a/Engine/Plugins/Experimental/PythonScriptPlugin/Source/PythonScriptPlugin/Private/PipInstall.h b/Engine/Plugins/Experimental/PythonScriptPlugin/Source/PythonScriptPlugin/Private/PipInstall.h
--- a/Engine/Plugins/Experimental/PythonScriptPlugin/Source/PythonScriptPlugin/Private/PipInstall.h
+++ b/Engine/Plugins/Experimental/PythonScriptPlugin/Source/PythonScriptPlugin/Private/PipInstall.h
@@ -21,6 +21,8 @@ public:
 	static const FString ExtraUrlsFilename;
 	static const FString ParsedRequirementsFilename;
 
+	static bool EnabledOnStartup();
+
 	static FString WritePluginsListing(TArray<TSharedRef<IPlugin>>& OutPythonPlugins);
 	static FString WritePluginDependencies(const TArray<TSharedRef<IPlugin>>& PythonPlugins, TArray<FString>& OutRequirements, TArray<FString>& OutExtraUrls);
 
diff --git a/Engine/Plugins/Experimental/PythonScriptPlugin/Source/PythonScriptPlugin/Private/PythonScriptPlugin.cpp b/Engine/Plugins/Experimental/PythonScriptPlugin/Source/PythonScriptPlugin/Private/PythonScriptPlugin.cpp
--- a/Engine/Plugins/Experimental/PythonScriptPlugin/Source/PythonScriptPlugin/Private/PythonScriptPlugin.cpp
+++ b/Engine/Plugins/Experimental/PythonScriptPlugin/Source/PythonScriptPlugin/Private/PythonScriptPlugin.cpp
@@ -1144,31 +1144,33 @@ void FPythonScriptPlugin::InitPipInstaller()
 	TArray<FString> ExtraUrls;
 	const FString ReqsInFile = FPipInstall::WritePluginDependencies(PythonPlugins, ReqInLines, ExtraUrls);
 
-	bool bRunOnStartup = GetDefault<UPythonScriptPluginSettings>()->bRunPipInstallOnStartup;
+	bool EnabledOnStart = FPipInstall::EnabledOnStartup();
 	if (ReqInLines.IsEmpty())
 	{
-		UE_CLOG(bRunOnStartup, LogPython, Display, TEXT("No enabled plugins with python dependencies found, skipping"));
+		UE_CLOG(EnabledOnStart, LogPython, Display, TEXT("No enabled plugins with python dependencies found, skipping"));
 		return;
 	}
 
 	// Just return immediately with warning if some python dependencies exist and pip install is disabled
-	bool bCmdLineDisable = FParse::Param(FCommandLine::Get(), TEXT("DisablePipInstall"));
-	if (bCmdLineDisable || !bRunOnStartup || GIsBuildMachine)
+	if (!EnabledOnStart)
 	{
+		bool bCmdLineDisable = FParse::Param(FCommandLine::Get(), TEXT("DisablePipInstall"));
 		if (bCmdLineDisable || GIsBuildMachine)
 		{
 			// Don't warn if disabled on cmd-line or is build process
 			UE_LOG(LogPython, Display, TEXT("Enabled plugins have python dependencies, install manually to: %s"), *PipSitePackagePath);
-			UE_LOG(LogPython, Display, TEXT("  See package requirements: % s"), *ReqsInFile);
+			UE_LOG(LogPython, Display, TEXT("  See package requirements: %s"), *ReqsInFile);
 		}
 		else
 		{
 			UE_LOG(LogPython, Warning, TEXT("Enabled plugins have python dependencies, enable 'Run Pip Install On Startup' or install manually to: %s"), *PipSitePackagePath);
-			UE_LOG(LogPython, Warning, TEXT("  See package requirements: % s"), *ReqsInFile);
+			UE_LOG(LogPython, Warning, TEXT("  See package requirements: %s"), *ReqsInFile);
 		}
 		return;
 	}
 
+	UE_LOG(LogPython, Display, TEXT("Preparing to install python dependencies into: %s"), *PipSitePackagePath);
+
 	FPipInstall::SetupPipEnv(Context);
 	FPipInstall::ParsePluginDependencies(ReqsInFile, Context);
 }
@@ -1177,9 +1179,7 @@ void FPythonScriptPlugin::RunPipInstaller()
 {
 	TRACE_CPUPROFILER_EVENT_SCOPE(FPythonScriptPlugin::RunPipInstaller);
 
-	bool bRunOnStartup = GetDefault<UPythonScriptPluginSettings>()->bRunPipInstallOnStartup;
-	bool bCmdLineDisable = FParse::Param(FCommandLine::Get(), TEXT("DisablePipInstall"));
-	if (bCmdLineDisable || !bRunOnStartup || GIsBuildMachine)
+	if (!FPipInstall::EnabledOnStartup())
 	{
 		return;
 	}
@@ -1208,6 +1208,7 @@ bool FPythonScriptPlugin::RunUBTPipAction(const FString& Action, const FText& De
 {
 	TRACE_CPUPROFILER_EVENT_SCOPE(FPythonScriptPlugin::RunUBTPipAction);
 
+	// TODO: Convert error codes to specific error messages for pip installer
 	int32 ExitCode;
 	const FString ProjectFileName = IFileManager::Get().ConvertToAbsolutePathForExternalAppForWrite(*FPaths::GetProjectFilePath());
 	const FString Args = FString::Printf(TEXT("%s %s %s -Project=\"%s\" -Mode=PipInstall -PythonInterpreter=\"%s\" -PipAction=%s -Progress")
