PipInstall: Update python plugin dependency parser to handle pip-installed vs. plugin-packaged depenencies. Additionally, leave matched dependencies in the requirements list to allow pip to detect dependent package version mismatch errors.

diff --git a/Engine/Plugins/Experimental/PythonScriptPlugin/Source/PythonScriptPlugin/Private/PipInstall.cpp b/Engine/Plugins/Experimental/PythonScriptPlugin/Source/PythonScriptPlugin/Private/PipInstall.cpp
--- a/Engine/Plugins/Experimental/PythonScriptPlugin/Source/PythonScriptPlugin/Private/PipInstall.cpp
+++ b/Engine/Plugins/Experimental/PythonScriptPlugin/Source/PythonScriptPlugin/Private/PipInstall.cpp
@@ -20,7 +20,7 @@
 
 // In order to keep editor startup time fast, check directly for this utils version (make sure to match with wheel version in PythonScriptPlugin/Content/Python/Lib/wheels)
 // NOTE: This version must also be changed in PipInstallMode.cs in order to support UBT functionality
-const FString FPipInstall::PipInstallUtilsVer = TEXT("0.1.3");
+const FString FPipInstall::PipInstallUtilsVer = TEXT("0.1.4");
 
 const FString FPipInstall::PluginsListingFilename = TEXT("pyreqs_plugins.list");
 const FString FPipInstall::PluginsSitePackageFilename = TEXT("plugin_site_package.pth");
@@ -186,7 +186,8 @@ bool FPipInstall::HasInstallLines(const TArray<FString>& RequirementLines)
 {
 	for (const FStringView Line : RequirementLines)
 	{
-		if (!Line.TrimStart().StartsWith(TCHAR('#')))
+		bool bCommentLine = Line.TrimStart().StartsWith(TCHAR('#'));
+		if (!bCommentLine && !Line.Contains(TEXT("# [pkg:check]")))
 		{
 			return true;
 		}
diff --git a/Engine/Source/Programs/UnrealBuildTool/Modes/PipInstallMode.cs b/Engine/Source/Programs/UnrealBuildTool/Modes/PipInstallMode.cs
--- a/Engine/Source/Programs/UnrealBuildTool/Modes/PipInstallMode.cs
+++ b/Engine/Source/Programs/UnrealBuildTool/Modes/PipInstallMode.cs
@@ -166,7 +166,7 @@ namespace UnrealBuildTool.Modes
 	{
 		// Don't bother to re-install pip install tools if this version is already installed
 		// NOTE: This version must also be changed in PipInstall.cpp in order to support editor startup process
-		private const string PipInstallUtilsVer = "0.1.3";
+		private const string PipInstallUtilsVer = "0.1.4";
 
 		// Generated from enabled plugins list
 		private const string PluginsListingFilename = "pyreqs_plugins.list";
