Fixed UnpackReturnValues failing to correctly handle a None return value from a Python implemented function

Added a test for this to test_wrapper_types.py

diff --git a/Engine/Plugins/Experimental/PythonScriptPlugin/Source/PythonScriptPlugin/Private/PyGenUtil.cpp b/Engine/Plugins/Experimental/PythonScriptPlugin/Source/PythonScriptPlugin/Private/PyGenUtil.cpp
--- a/Engine/Plugins/Experimental/PythonScriptPlugin/Source/PythonScriptPlugin/Private/PyGenUtil.cpp
+++ b/Engine/Plugins/Experimental/PythonScriptPlugin/Source/PythonScriptPlugin/Private/PyGenUtil.cpp
@@ -1101,6 +1101,12 @@ bool UnpackReturnValues(PyObject* InRetVals, const FOutParmRec* InOutputParms, c
 		const bool bReturnValue = InRetVals != Py_None;
 		BoolReturn->SetPropertyValue(OutParamRec->PropAddr, bReturnValue);
 
+		if (!bReturnValue)
+		{
+			// None was returned, so there's nothing to unpack into any additional output arguments
+			return true;
+		}
+
 		OutParamRec = OutParamRec->NextOutParm; // Start unpacking at the 1st out value
 		check(OutParamRec);
 	}
diff --git a/Engine/Plugins/Experimental/PythonScriptPlugin/Source/PythonScriptPlugin/Private/PyTest.cpp b/Engine/Plugins/Experimental/PythonScriptPlugin/Source/PythonScriptPlugin/Private/PyTest.cpp
--- a/Engine/Plugins/Experimental/PythonScriptPlugin/Source/PythonScriptPlugin/Private/PyTest.cpp
+++ b/Engine/Plugins/Experimental/PythonScriptPlugin/Source/PythonScriptPlugin/Private/PyTest.cpp
@@ -107,6 +107,11 @@ int32 UPyTestObject::CallFuncBlueprintImplementable(const int32 InValue) const
 	return FuncBlueprintImplementable(InValue);
 }
 
+bool UPyTestObject::CallFuncBlueprintImplementablePackedGetter(int32& OutValue) const
+{
+	return FuncBlueprintImplementablePackedGetter(OutValue);
+}
+
 int32 UPyTestObject::CallFuncBlueprintNative(const int32 InValue) const
 {
 	return FuncBlueprintNative(InValue);
diff --git a/Engine/Plugins/Experimental/PythonScriptPlugin/Source/PythonScriptPlugin/Private/PyTest.h b/Engine/Plugins/Experimental/PythonScriptPlugin/Source/PythonScriptPlugin/Private/PyTest.h
--- a/Engine/Plugins/Experimental/PythonScriptPlugin/Source/PythonScriptPlugin/Private/PyTest.h
+++ b/Engine/Plugins/Experimental/PythonScriptPlugin/Source/PythonScriptPlugin/Private/PyTest.h
@@ -225,6 +225,9 @@ public:
 	UFUNCTION(BlueprintImplementableEvent, Category = "Python|Internal")
 	int32 FuncBlueprintImplementable(const int32 InValue) const;
 
+	UFUNCTION(BlueprintImplementableEvent, Category = "Python|Internal")
+	bool FuncBlueprintImplementablePackedGetter(int32& OutValue) const;
+
 	UFUNCTION(BlueprintNativeEvent, Category = "Python|Internal")
 	int32 FuncBlueprintNative(const int32 InValue) const;
 
@@ -234,6 +237,9 @@ public:
 	UFUNCTION(BlueprintCallable, Category = "Python|Internal")
 	int32 CallFuncBlueprintImplementable(const int32 InValue) const;
 
+	UFUNCTION(BlueprintCallable, Category = "Python|Internal")
+	bool CallFuncBlueprintImplementablePackedGetter(int32& OutValue) const;
+
 	UFUNCTION(BlueprintCallable, Category = "Python|Internal")
 	int32 CallFuncBlueprintNative(const int32 InValue) const;
 
