PipInstall: Non-UBT code path for pip install from in-editor

diff --git a/Engine/Plugins/Experimental/PythonScriptPlugin/Source/PythonScriptPlugin/Private/PipInstall.cpp b/Engine/Plugins/Experimental/PythonScriptPlugin/Source/PythonScriptPlugin/Private/PipInstall.cpp
--- a/Engine/Plugins/Experimental/PythonScriptPlugin/Source/PythonScriptPlugin/Private/PipInstall.cpp
+++ b/Engine/Plugins/Experimental/PythonScriptPlugin/Source/PythonScriptPlugin/Private/PipInstall.cpp
@@ -18,6 +18,87 @@
 
 #if WITH_PYTHON
 
+// Simple interface for parsing cmd output to update slowtask progress
+// Similar to FFeedbackContextMarkup, but supports arbitrary line parsing
+class IProgressParser
+{
+public:
+	// Get a total work estimate
+	virtual int GetTotalWork() = 0;
+	// Parse line and update status/progress (return true to eat the output and not log)
+	virtual bool UpdateStatus(const FString& ChkLine, FSlowTask& Context) = 0;
+};
+
+class FPipProgressParser : public IProgressParser
+{
+public:
+	FPipProgressParser(int InRequirementsCount)
+	: RequirementsDone(0)
+	, RequirementsCount(FMath::Max(InRequirementsCount,1.0f))
+	{}
+
+	virtual int GetTotalWork() override
+	{
+		return RequirementsCount;
+	}
+
+	virtual bool UpdateStatus(const FString& ChkLine, FSlowTask& Task) override
+	{
+		FString TrimLine = ChkLine.TrimStartAndEnd();
+		// Just log if it's not a status update line
+		if (!CheckUpdateMatch(TrimLine))
+		{
+			return false;
+		}
+
+		// TODO: Pass in specific requirements to update status lines more accurately
+		FString StatusStr = ReplaceUpdateStrs(TrimLine);
+		Task.EnterProgressFrame(1.0f / (RequirementsCount + PadCount), FText::FromString(StatusStr));
+
+		// Exponentially approach 100% if steps goes above estimate
+		RequirementsDone += 1.0f;
+		RequirementsCount = FMath::Max(RequirementsCount, RequirementsDone + 1.0f);
+
+		return false;
+	}
+
+private:
+	static bool CheckUpdateMatch(const FString& Line)
+	{
+		for (const FString& ChkMatch : MatchStatusStrs)
+		{
+			if (Line.StartsWith(ChkMatch))
+			{
+				return true;
+			}
+		}
+
+		return false;
+	}
+
+	static FString ReplaceUpdateStrs(const FString& Line)
+	{
+		FString RepLine = Line;
+		for (const TPair<FString, FString>& ReplaceMap : LogReplaceStrs)
+		{
+			RepLine = RepLine.Replace(*ReplaceMap.Key, *ReplaceMap.Value, ESearchCase::CaseSensitive);
+		}
+
+		return RepLine;
+	}
+
+	float RequirementsDone;
+	float RequirementsCount;
+
+	static const int PadCount = 2;
+	static const TArray<FString> MatchStatusStrs;
+	static const TMap<FString,FString> LogReplaceStrs;
+};
+
+const TArray<FString> FPipProgressParser::MatchStatusStrs = {TEXT("Requirement"), TEXT("Downloading"), TEXT("Using"), TEXT("Installing")};
+const TMap<FString,FString> FPipProgressParser::LogReplaceStrs = {{TEXT("Installing collected packages:"), TEXT("Installing collected python package dependencies:")}};
+
+
 // In order to keep editor startup time fast, check directly for this utils version (make sure to match with wheel version in PythonScriptPlugin/Content/Python/Lib/wheels)
 // NOTE: This version must also be changed in PipInstallMode.cs in order to support UBT functionality
 const FString FPipInstall::PipInstallUtilsVer = TEXT("0.1.4");
@@ -204,18 +285,72 @@ FString FPipInstall::ParsePluginDependencies(const FString& MergedInRequirements
 	return FPaths::ConvertRelativePathToFull(ParsedReqsFile);
 }
 
-bool FPipInstall::HasInstallLines(const TArray<FString>& RequirementLines)
+bool FPipInstall::RunPipInstall(FFeedbackContext* Context, bool bOfflineOnly, const FString& ForceIndexUrl)
 {
+	TRACE_CPUPROFILER_EVENT_SCOPE(FPipInstall::RunPipInstall);
+
+	const FString PipInstallPath = GetPipInstallPath();
+	const FString VenvInterp = GetVenvInterpreter(PipInstallPath);
+
+	const FString ParsedReqsFile = PipInstallPath / ParsedRequirementsFilename;
+	const FString ExtraUrlsFile = PipInstallPath / ExtraUrlsFilename;
+
+	if (!FPaths::FileExists(ParsedReqsFile))
+	{
+		return true;
+	}
+
+	TArray<FString> ParsedReqLines;
+	if ( !FFileHelper::LoadFileToStringArray(ParsedReqLines, *ParsedReqsFile) )
+	{
+		return false;
+	}
+
+	int ReqCount = CountInstallLines(ParsedReqLines);
+
+	TArray<FString> ExtraUrls;
+	if (FPaths::FileExists(ExtraUrlsFile))
+	{
+		FFileHelper::LoadFileToStringArray(ExtraUrls, *ExtraUrlsFile);
+	}
+
+	FString Cmd = TEXT("-m pip install --disable-pip-version-check --only-binary=:all:");
+	if (bOfflineOnly)
+	{
+		Cmd += TEXT(" --no-index");
+	}
+	else if (!ForceIndexUrl.IsEmpty())
+	{
+		Cmd += TEXT("--index-url ") + ForceIndexUrl;
+	}
+	else if (!ExtraUrls.IsEmpty())
+	{
+		for (const FString& Url: ExtraUrls)
+		{
+			Cmd += TEXT(" --extra-index-url ") + Url;
+		}
+	}
+
+	Cmd += " -r \"" + ParsedReqsFile + "\"";
+
+	TSharedPtr<IProgressParser> ProgParser = MakeShared<FPipProgressParser>(ReqCount);
+	int32 Result = RunPythonCmd(LOCTEXT("PipInstall.InstallRequirements", "Installing pip requirements..."), VenvInterp, Cmd, Context, ProgParser);
+	return (Result == 0);
+}
+
+int FPipInstall::CountInstallLines(const TArray<FString>& RequirementLines)
+{
+	int Count = 0;
 	for (const FStringView Line : RequirementLines)
 	{
 		bool bCommentLine = Line.TrimStart().StartsWith(TCHAR('#'));
 		if (!bCommentLine && !Line.Contains(TEXT("# [pkg:check]")))
 		{
-			return true;
+			Count += 1;
 		}
 	}
 
-	return false;
+	return Count;
 }
 
 FString FPipInstall::GetPipInstallPath()
@@ -257,21 +392,23 @@ bool FPipInstall::CheckPipInstallUtils(const FString& VenvInterp, FFeedbackConte
 	return (RunPythonCmd(LOCTEXT("PipInstall.CheckPipInstallUtils", "Check pip install utils installed"), VenvInterp, Cmd, Context) == 0);
 }
 
-int32 FPipInstall::RunPythonCmd(const FText& Description, const FString& PythonInterp, const FString& Cmd, FFeedbackContext* Context)
+int32 FPipInstall::RunPythonCmd(const FText& Description, const FString& PythonInterp, const FString& Cmd, FFeedbackContext* Context, TSharedPtr<IProgressParser> CmdParser)
 {
 	TRACE_CPUPROFILER_EVENT_SCOPE(FPipInstall::RunPythonCmd);
 
 	UE_LOG(LogPython, Log, TEXT("Running python command: python %s"), *Cmd);
 
 	int32 Result = 0;
-	RunLoggedSubprocess(Description, FPaths::ConvertRelativePathToFull(PythonInterp), Cmd, Context, &Result);
+	RunLoggedSubprocess(&Result, Description, FPaths::ConvertRelativePathToFull(PythonInterp), Cmd, Context, CmdParser);
 
 	return Result;
 }
 
-bool FPipInstall::RunLoggedSubprocess(const FText& Description, const FString& URL, const FString& Params, FFeedbackContext* Context, int32* OutExitCode)
+bool FPipInstall::RunLoggedSubprocess(int32* OutExitCode, const FText& Description, const FString& URL, const FString& Params, FFeedbackContext* Context, TSharedPtr<IProgressParser> CmdParser)
 {
-	FScopedSlowTask SubprocessTask(0, Description, true, *Context);
+	int AmountOfWork = (CmdParser.IsValid()) ? CmdParser->GetTotalWork() : 0;
+	FScopedSlowTask SubprocessTask(1.0f, Description, true, *Context);
+	SubprocessTask.MakeDialog();
 
 	// Create a read and write pipe for the child process
 	void* StdOutPipeRead = nullptr;
@@ -299,7 +436,12 @@ bool FPipInstall::RunLoggedSubprocess(const FText& Description, const FString& U
 				FString Line = BufferedText.Left(EndOfLineIdx);
 				Line.RemoveFromEnd(TEXT("\r"), ESearchCase::CaseSensitive);
 
-				Context->Log(LogPython.GetCategoryName(), ELogVerbosity::Log, Line);
+				// Always log if no output parser, also log if UpdateStatus returns false
+				if (!CmdParser.IsValid() || !CmdParser->UpdateStatus(Line, SubprocessTask))
+				{
+					Context->Log(LogPython.GetCategoryName(), ELogVerbosity::Log, Line);
+				}
+
 				BufferedText.MidInline(EndOfLineIdx + 1, MAX_int32, false);
 			}
 
diff --git a/Engine/Plugins/Experimental/PythonScriptPlugin/Source/PythonScriptPlugin/Private/PipInstall.h b/Engine/Plugins/Experimental/PythonScriptPlugin/Source/PythonScriptPlugin/Private/PipInstall.h
--- a/Engine/Plugins/Experimental/PythonScriptPlugin/Source/PythonScriptPlugin/Private/PipInstall.h
+++ b/Engine/Plugins/Experimental/PythonScriptPlugin/Source/PythonScriptPlugin/Private/PipInstall.h
@@ -11,6 +11,8 @@
 class FJsonObject;
 class FFeedbackContext;
 
+class IProgressParser;
+
 class FPipInstall
 {
 public:
@@ -30,16 +32,17 @@ public:
 
 	static void SetupPipEnv(FFeedbackContext* Context, bool bForceRebuild = false);
 	static FString ParsePluginDependencies(const FString& MergedInRequirementsFile, FFeedbackContext* Context);
+	static bool RunPipInstall(FFeedbackContext* Context, bool bOfflineOnly = false, const FString& ForceIndexUrl = TEXT(""));
 
-	static bool HasInstallLines(const TArray<FString>& RequirementLines);
+	static int CountInstallLines(const TArray<FString>& RequirementLines);
 
 	static FString GetPipInstallPath();
 
 private:
 	static void SetupPipInstallUtils(const FString& VenvInterp, FFeedbackContext* Context);
 	static bool CheckPipInstallUtils(const FString& VenvInterp, FFeedbackContext* Context);
-	static int32 RunPythonCmd(const FText& Description, const FString& VenvInterp, const FString& Cmd, FFeedbackContext* Context);
-	static bool RunLoggedSubprocess(const FText& Description, const FString& URL, const FString& Params, FFeedbackContext* Context, int32* OutExitCode);\
+	static int32 RunPythonCmd(const FText& Description, const FString& PythonInterp, const FString& Cmd, FFeedbackContext* Context, TSharedPtr<IProgressParser> CmdParser = nullptr);
+	static bool RunLoggedSubprocess(int32* OutExitCode, const FText& Description, const FString& URL, const FString& Params, FFeedbackContext* Context, TSharedPtr<IProgressParser> CmdParser);
 
 	static FString GetPythonScriptPluginPath();
 	static FString ParseVenvVersion(const FString& InstallPath);
diff --git a/Engine/Plugins/Experimental/PythonScriptPlugin/Source/PythonScriptPlugin/Private/PythonScriptPlugin.cpp b/Engine/Plugins/Experimental/PythonScriptPlugin/Source/PythonScriptPlugin/Private/PythonScriptPlugin.cpp
--- a/Engine/Plugins/Experimental/PythonScriptPlugin/Source/PythonScriptPlugin/Private/PythonScriptPlugin.cpp
+++ b/Engine/Plugins/Experimental/PythonScriptPlugin/Source/PythonScriptPlugin/Private/PythonScriptPlugin.cpp
@@ -1188,41 +1188,21 @@ void FPythonScriptPlugin::RunPipInstaller()
 	const FString ParsedReqsFile = FPaths::ConvertRelativePathToFull(FPipInstall::GetPipInstallPath() / FPipInstall::ParsedRequirementsFilename);
 
 	TArray<FString> ParsedReqLines;
-	if (!FPaths::FileExists(ParsedReqsFile) || !FFileHelper::LoadFileToStringArray(ParsedReqLines, *ParsedReqsFile) || !FPipInstall::HasInstallLines(ParsedReqLines))
+	if (!FPaths::FileExists(ParsedReqsFile) || !FFileHelper::LoadFileToStringArray(ParsedReqLines, *ParsedReqsFile) || (FPipInstall::CountInstallLines(ParsedReqLines) == 0))
 	{
 		return;
 	}
 
 	FFeedbackContext* Context = GWarn;
-	FScopedSlowTask PipInstallTask(0, LOCTEXT("PipInstall.RunInstall", "Installing Python Dependencies..."), true, *Context);
-	//PipInstallTask.MakeDialog(true);
 
 	// Run install of all python dependencies for enabled plugins
-	if (!RunUBTPipAction("InstallNoRegen", LOCTEXT("PipInstall.UBTInstall", "Pip Installing Plugin Dependencies..."), Context))
+	if (!FPipInstall::RunPipInstall(Context))
 	{
 		UE_LOG(LogPython, Warning, TEXT("Unable to install plugin python dependencies"));
 		return;
 	}
 }
 
-bool FPythonScriptPlugin::RunUBTPipAction(const FString& Action, const FText& Description, FFeedbackContext* Context)
-{
-	TRACE_CPUPROFILER_EVENT_SCOPE(FPythonScriptPlugin::RunUBTPipAction);
-
-	// TODO: Convert error codes to specific error messages for pip installer
-	int32 ExitCode;
-	const FString ProjectFileName = IFileManager::Get().ConvertToAbsolutePathForExternalAppForWrite(*FPaths::GetProjectFilePath());
-	const FString Args = FString::Printf(TEXT("%s %s %s -Project=\"%s\" -Mode=PipInstall -PythonInterpreter=\"%s\" -PipAction=%s -Progress")
-		, FPlatformMisc::GetUBTTargetName()
-		, FPlatformMisc::GetUBTPlatform()
-		, FModuleManager::Get().GetUBTConfiguration()
-		, *ProjectFileName
-		, *PyUtil::GetInterpreterExecutablePath()
-		, *Action);
-
-	return FDesktopPlatformModule::Get()->RunUnrealBuildTool(Description, FPaths::RootDir(), Args, Context, ExitCode);
-}
-
 FString FPythonScriptPlugin::GetPipSitePackagesPath()
 {
 	const FString VenvPath = FPipInstall::GetPipInstallPath();
