PipInstall: Remove leftover __pycache__ dir hierarchies in pip-managed plugin site-package folders.

diff --git a/Engine/Plugins/Experimental/PythonScriptPlugin/Source/PythonScriptPlugin/Private/PipInstall.cpp b/Engine/Plugins/Experimental/PythonScriptPlugin/Source/PythonScriptPlugin/Private/PipInstall.cpp
--- a/Engine/Plugins/Experimental/PythonScriptPlugin/Source/PythonScriptPlugin/Private/PipInstall.cpp
+++ b/Engine/Plugins/Experimental/PythonScriptPlugin/Source/PythonScriptPlugin/Private/PipInstall.cpp
@@ -216,6 +216,94 @@ FString FPipInstall::WritePluginDependencies(const TArray<TSharedRef<IPlugin>>&
 	return MergedReqsFile;
 }
 
+class FCheckOrphanDirVisitor : public IPlatformFile::FDirectoryVisitor
+{
+public:
+	FCheckOrphanDirVisitor()
+	: IPlatformFile::FDirectoryVisitor()
+	, bOrphan(true)
+	{}
+
+	virtual bool Visit(const TCHAR* FilenameOrDirectory, bool bIsDir) override
+	{
+		if (!bIsDir)
+		{
+			bOrphan = false;
+			return true;
+		}
+
+		//// Short circuit recursion if we don't allow deleting sub-hierarchies and this traversal level is already non-orphan
+		//if (!bAllowDeleteSubdirs && !bOrphan)
+		//{
+		//	return true;
+		//}
+
+		// Always treat __pycache__ dir as orphan but don't directly delete them unless full parent is also orphan (nothing but empty or __pycache__ dirs)
+		const FString DirPath(FilenameOrDirectory);
+		if (DirPath.EndsWith(TEXT("__pycache__")))
+		{
+			return true;
+		}
+
+		FCheckOrphanDirVisitor SubDirVisit;
+		IPlatformFile& PlatformFile = FPlatformFileManager::Get().GetPlatformFile();
+		bool res = PlatformFile.IterateDirectory(FilenameOrDirectory, SubDirVisit);
+
+		bOrphan = bOrphan && SubDirVisit.bOrphan;
+		if (SubDirVisit.bOrphan)
+		{
+			Orphans.Add(FilenameOrDirectory);
+		}
+		//else if (bAllowDeleteSubdirs)
+		//{
+		//	Orphans.Append(SubDirVisit.Orphans);
+		//}
+
+		return res;
+	}
+
+public:
+	bool bOrphan;
+	TArray<FString> Orphans;
+};
+
+// Remove orphan path hierarchies (hierarchies with only __pycache__ or empty dirs)
+// Only runs for <PluginDir>/Content/Python/Lib/* subdirectories for plugins with
+// Pip PythonRequirements uplugin section
+void FPipInstall::CheckRemoveOrphanedPackages(const FString& SitePackagesPath)
+{
+	TRACE_CPUPROFILER_EVENT_SCOPE(FPipInstall::CheckRemoveOrphanedPackages);
+
+	if (!FPaths::DirectoryExists(SitePackagesPath))
+	{
+		return;
+	}
+	
+	// NOTE: FCheckOrphanDirVisitor should only return top-level orphan hierarchies for removal (all or nothing)
+	FCheckOrphanDirVisitor DirVisit;
+	IPlatformFile& PlatformFile = FPlatformFileManager::Get().GetPlatformFile();
+	if (!PlatformFile.IterateDirectory(*SitePackagesPath, DirVisit))
+	{
+		return;
+	}
+
+	if (DirVisit.bOrphan)
+	{
+		// Remove the entire site-packages dir if everything beneath is orphaned
+		UE_LOG(LogPython, Log, TEXT("PipInstall found orphan plugin site-package directory: %s (removing)"), *SitePackagesPath);
+		PlatformFile.DeleteDirectoryRecursively(*SitePackagesPath);
+	}
+	else
+	{
+		// Only remove specifically orphaned subdirs if there are some valid hierarchies in site-packages
+		for (const FString& OrphanDir : DirVisit.Orphans)
+		{
+			UE_LOG(LogPython, Log, TEXT("PipInstall found orphan plugin site-package directory: %s (removing)"), *OrphanDir);
+			PlatformFile.DeleteDirectoryRecursively(*OrphanDir);
+		}
+	}
+}
+
 void FPipInstall::CheckInvalidPipEnv()
 {
 	TRACE_CPUPROFILER_EVENT_SCOPE(FPipInstall::CheckInvalidPipEnv);
@@ -270,6 +358,20 @@ void FPipInstall::SetupPipEnv(FFeedbackContext* Context, bool bForceRebuild /* =
 	SetupPipInstallUtils(VenvInterp, Context);
 }
 
+void FPipInstall::RemoveParsedDependencyFiles()
+{
+	TRACE_CPUPROFILER_EVENT_SCOPE(FPipInstall::ParsePluginDependencies);
+
+	const FString PipInstallPath = GetPipInstallPath();
+	const FString ParsedReqsFile = PipInstallPath / ParsedRequirementsFilename;
+
+	if (FPaths::FileExists(ParsedReqsFile))
+	{
+		IPlatformFile& PlatformFile = FPlatformFileManager::Get().GetPlatformFile();
+		PlatformFile.DeleteFile(*ParsedReqsFile);
+	}
+}
+
 FString FPipInstall::ParsePluginDependencies(const FString& MergedInRequirementsFile, FFeedbackContext* Context)
 {
 	TRACE_CPUPROFILER_EVENT_SCOPE(FPipInstall::ParsePluginDependencies);
diff --git a/Engine/Plugins/Experimental/PythonScriptPlugin/Source/PythonScriptPlugin/Private/PipInstall.h b/Engine/Plugins/Experimental/PythonScriptPlugin/Source/PythonScriptPlugin/Private/PipInstall.h
--- a/Engine/Plugins/Experimental/PythonScriptPlugin/Source/PythonScriptPlugin/Private/PipInstall.h
+++ b/Engine/Plugins/Experimental/PythonScriptPlugin/Source/PythonScriptPlugin/Private/PipInstall.h
@@ -25,12 +25,14 @@ public:
 
 	static bool EnabledOnStartup();
 
+	static void CheckRemoveOrphanedPackages(const FString& SitePackagesPath);
 	static void CheckInvalidPipEnv();
 
 	static FString WritePluginsListing(TArray<TSharedRef<IPlugin>>& OutPythonPlugins);
 	static FString WritePluginDependencies(const TArray<TSharedRef<IPlugin>>& PythonPlugins, TArray<FString>& OutRequirements, TArray<FString>& OutExtraUrls);
 
 	static void SetupPipEnv(FFeedbackContext* Context, bool bForceRebuild = false);
+	static void RemoveParsedDependencyFiles();
 	static FString ParsePluginDependencies(const FString& MergedInRequirementsFile, FFeedbackContext* Context);
 	static bool RunPipInstall(FFeedbackContext* Context, bool bOfflineOnly = false, const FString& ForceIndexUrl = TEXT(""));
 
diff --git a/Engine/Plugins/Experimental/PythonScriptPlugin/Source/PythonScriptPlugin/Private/PythonScriptPlugin.cpp b/Engine/Plugins/Experimental/PythonScriptPlugin/Source/PythonScriptPlugin/Private/PythonScriptPlugin.cpp
--- a/Engine/Plugins/Experimental/PythonScriptPlugin/Source/PythonScriptPlugin/Private/PythonScriptPlugin.cpp
+++ b/Engine/Plugins/Experimental/PythonScriptPlugin/Source/PythonScriptPlugin/Private/PythonScriptPlugin.cpp
@@ -1141,6 +1141,14 @@ void FPythonScriptPlugin::InitPipInstaller()
 	TArray<TSharedRef<IPlugin>> PythonPlugins;
 	FPipInstall::WritePluginsListing(PythonPlugins);
 
+	for (const TSharedRef<IPlugin>& PyPlugin : PythonPlugins)
+	{
+		// Remove leftover __pycache__ folders from plugins that use pip, but previously used packaged dependencies
+		const FString LibDir = PyPlugin->GetContentDir() / TEXT("Python") / TEXT("Lib");
+		FPipInstall::CheckRemoveOrphanedPackages(LibDir / TEXT("site-packages"));
+		FPipInstall::CheckRemoveOrphanedPackages(LibDir / FPlatformMisc::GetUBTPlatform() / TEXT("site-packages"));
+	}
+
 	TArray<FString> ReqInLines;
 	TArray<FString> ExtraUrls;
 	const FString ReqsInFile = FPipInstall::WritePluginDependencies(PythonPlugins, ReqInLines, ExtraUrls);
@@ -1149,6 +1157,7 @@ void FPythonScriptPlugin::InitPipInstaller()
 	if (ReqInLines.IsEmpty())
 	{
 		UE_CLOG(EnabledOnStart, LogPython, Display, TEXT("No enabled plugins with python dependencies found, skipping"));
+		FPipInstall::RemoveParsedDependencyFiles();
 		return;
 	}
 
