PipInstall: Fix update progress in running pip install so that progress frame work never exceeds total work for task.

diff --git a/Engine/Plugins/Experimental/PythonScriptPlugin/Source/PythonScriptPlugin/Private/PipInstall.cpp b/Engine/Plugins/Experimental/PythonScriptPlugin/Source/PythonScriptPlugin/Private/PipInstall.cpp
--- a/Engine/Plugins/Experimental/PythonScriptPlugin/Source/PythonScriptPlugin/Private/PipInstall.cpp
+++ b/Engine/Plugins/Experimental/PythonScriptPlugin/Source/PythonScriptPlugin/Private/PipInstall.cpp
@@ -26,15 +26,15 @@ public:
 	// Get a total work estimate
 	virtual int GetTotalWork() = 0;
 	// Parse line and update status/progress (return true to eat the output and not log)
-	virtual bool UpdateStatus(const FString& ChkLine, FSlowTask& Context) = 0;
+	virtual bool UpdateStatus(const FString& ChkLine, FSlowTask& Task) = 0;
 };
 
 class FPipProgressParser : public IProgressParser
 {
 public:
-	FPipProgressParser(int InRequirementsCount)
+	FPipProgressParser(int GuessRequirementsCount)
 	: RequirementsDone(0)
-	, RequirementsCount(FMath::Max(InRequirementsCount,1.0f))
+	, RequirementsCount(FMath::Max(GuessRequirementsCount,1.0f))
 	{}
 
 	virtual int GetTotalWork() override
@@ -51,13 +51,15 @@ public:
 			return false;
 		}
 
+		// Exponentially approach 100% if steps goes above estimate
+		float ProgLeft = RequirementsCount - RequirementsDone;
+		float NextWork = FMath::Clamp(0.9*ProgLeft, 0.0f, 1.0f);
+
 		// TODO: Pass in specific requirements to update status lines more accurately
 		FString StatusStr = ReplaceUpdateStrs(TrimLine);
-		Task.EnterProgressFrame(1.0f / (RequirementsCount + PadCount), FText::FromString(StatusStr));
+		Task.EnterProgressFrame(NextWork, FText::FromString(StatusStr));
 
-		// Exponentially approach 100% if steps goes above estimate
-		RequirementsDone += 1.0f;
-		RequirementsCount = FMath::Max(RequirementsCount, RequirementsDone + 1.0f);
+		RequirementsDone += NextWork;
 
 		return false;
 	}
@@ -90,7 +92,6 @@ private:
 	float RequirementsDone;
 	float RequirementsCount;
 
-	static const int PadCount = 2;
 	static const TArray<FString> MatchStatusStrs;
 	static const TMap<FString,FString> LogReplaceStrs;
 };
@@ -509,7 +510,8 @@ int32 FPipInstall::RunPythonCmd(const FText& Description, const FString& PythonI
 bool FPipInstall::RunLoggedSubprocess(int32* OutExitCode, const FText& Description, const FString& URL, const FString& Params, FFeedbackContext* Context, TSharedPtr<IProgressParser> CmdParser)
 {
 	int AmountOfWork = (CmdParser.IsValid()) ? CmdParser->GetTotalWork() : 0;
-	FScopedSlowTask SubprocessTask(1.0f, Description, true, *Context);
+
+	FScopedSlowTask SubprocessTask(AmountOfWork, Description, true, *Context);
 	SubprocessTask.MakeDialog();
 
 	// Create a read and write pipe for the child process
@@ -555,8 +557,9 @@ bool FPipInstall::RunLoggedSubprocess(int32* OutExitCode, const FText& Descripti
 	else
 	{
 		Context->CategorizedLogf(LogPython.GetCategoryName(), ELogVerbosity::Warning, TEXT("Couldn't create process '%s'"), *URL);
-		return false;
 	}
+
+	return false;
 }
 
 
