PipInstall: Fix unreal plugins .pth file location for linux/macos.

diff --git a/Engine/Plugins/Experimental/PythonScriptPlugin/Source/PythonScriptPlugin/Private/PipInstall.cpp b/Engine/Plugins/Experimental/PythonScriptPlugin/Source/PythonScriptPlugin/Private/PipInstall.cpp
--- a/Engine/Plugins/Experimental/PythonScriptPlugin/Source/PythonScriptPlugin/Private/PipInstall.cpp
+++ b/Engine/Plugins/Experimental/PythonScriptPlugin/Source/PythonScriptPlugin/Private/PipInstall.cpp
@@ -461,6 +461,18 @@ FString FPipInstall::GetPipInstallPath()
 	return FPaths::ProjectIntermediateDir() / TEXT("PipInstall");
 }
 
+FString FPipInstall::GetPipSitePackagesPath()
+{
+	const FString VenvPath = GetPipInstallPath();
+#if PLATFORM_WINDOWS
+	return VenvPath / TEXT("Lib") / TEXT("site-packages");
+#elif PLATFORM_MAC || PLATFORM_LINUX
+	return VenvPath / TEXT("lib") / FString::Printf(TEXT("python%d.%d"), PY_MAJOR_VERSION, PY_MINOR_VERSION) / TEXT("site-packages");
+#else
+	static_assert(false, "Python not supported on this platform!");
+#endif
+}
+
 
 void FPipInstall::SetupPipInstallUtils(const FString& VenvInterp, FFeedbackContext* Context)
 {
diff --git a/Engine/Plugins/Experimental/PythonScriptPlugin/Source/PythonScriptPlugin/Private/PipInstall.h b/Engine/Plugins/Experimental/PythonScriptPlugin/Source/PythonScriptPlugin/Private/PipInstall.h
--- a/Engine/Plugins/Experimental/PythonScriptPlugin/Source/PythonScriptPlugin/Private/PipInstall.h
+++ b/Engine/Plugins/Experimental/PythonScriptPlugin/Source/PythonScriptPlugin/Private/PipInstall.h
@@ -39,6 +39,7 @@ public:
 	static int CountInstallLines(const TArray<FString>& RequirementLines);
 
 	static FString GetPipInstallPath();
+	static FString GetPipSitePackagesPath();
 
 private:
 	static void SetupPipInstallUtils(const FString& VenvInterp, FFeedbackContext* Context);
diff --git a/Engine/Plugins/Experimental/PythonScriptPlugin/Source/PythonScriptPlugin/Private/PythonScriptPlugin.cpp b/Engine/Plugins/Experimental/PythonScriptPlugin/Source/PythonScriptPlugin/Private/PythonScriptPlugin.cpp
--- a/Engine/Plugins/Experimental/PythonScriptPlugin/Source/PythonScriptPlugin/Private/PythonScriptPlugin.cpp
+++ b/Engine/Plugins/Experimental/PythonScriptPlugin/Source/PythonScriptPlugin/Private/PythonScriptPlugin.cpp
@@ -984,7 +984,7 @@ void FPythonScriptPlugin::InitializePython()
 #endif // WITH_EDITOR
 
 		// Add Pip UBT install path to site-packages if it exists
-		const FString PipSitePackagePath = GetPipSitePackagesPath();
+		const FString PipSitePackagePath = FPipInstall::GetPipSitePackagesPath();
 		if (FPaths::DirectoryExists(PipSitePackagePath))
 		{
 			PyUtil::AddSitePackagesPath(PipSitePackagePath);
@@ -1135,7 +1135,7 @@ void FPythonScriptPlugin::InitPipInstaller()
 	FScopedSlowTask PipInstallTask(0, LOCTEXT("PipInstall.RunInit", "Running Pip Init Tasks..."), true, *Context);
 
 	FPipInstall::CheckInvalidPipEnv();
-	const FString PipSitePackagePath = FPaths::ConvertRelativePathToFull(GetPipSitePackagesPath());
+	const FString PipSitePackagePath = FPaths::ConvertRelativePathToFull(FPipInstall::GetPipSitePackagesPath());
 
 	// Generate the input listing files of plugins with python dependencies and the listing of all requirements (installed or not)
 	TArray<TSharedRef<IPlugin>> PythonPlugins;
@@ -1212,18 +1212,6 @@ void FPythonScriptPlugin::RunPipInstaller()
 	}
 }
 
-FString FPythonScriptPlugin::GetPipSitePackagesPath()
-{
-	const FString VenvPath = FPipInstall::GetPipInstallPath();
-#if PLATFORM_WINDOWS
-	return VenvPath / TEXT("Lib") / TEXT("site-packages");
-#elif PLATFORM_MAC || PLATFORM_LINUX
-	return VenvPath / TEXT("lib") / FString::Printf(TEXT("python%d.%d"), PY_MAJOR_VERSION, PY_MINOR_VERSION) / TEXT("site-packages");
-#else
-	static_assert(false, "Python not supported on this platform!");
-#endif
-}
-
 void FPythonScriptPlugin::RequestStubCodeGeneration()
 {
 	// Ignore requests made before the startup scripts have ran
diff --git a/Engine/Plugins/Experimental/PythonScriptPlugin/Source/PythonScriptPlugin/Private/PythonScriptPlugin.h b/Engine/Plugins/Experimental/PythonScriptPlugin/Source/PythonScriptPlugin/Private/PythonScriptPlugin.h
--- a/Engine/Plugins/Experimental/PythonScriptPlugin/Source/PythonScriptPlugin/Private/PythonScriptPlugin.h
+++ b/Engine/Plugins/Experimental/PythonScriptPlugin/Source/PythonScriptPlugin/Private/PythonScriptPlugin.h
@@ -168,8 +168,6 @@ private:
 
 	bool RunUBTPipAction(const FString& Action, const FText& Description, FFeedbackContext* Context);
 
-	FString GetPipSitePackagesPath();
-
 	void RequestStubCodeGeneration();
 
 	void GenerateStubCode();
