PipInstall: Updated fix for .pth files on linux/macos

diff --git a/Engine/Plugins/Experimental/PythonScriptPlugin/Source/PythonScriptPlugin/Private/PipInstall.cpp b/Engine/Plugins/Experimental/PythonScriptPlugin/Source/PythonScriptPlugin/Private/PipInstall.cpp
--- a/Engine/Plugins/Experimental/PythonScriptPlugin/Source/PythonScriptPlugin/Private/PipInstall.cpp
+++ b/Engine/Plugins/Experimental/PythonScriptPlugin/Source/PythonScriptPlugin/Private/PipInstall.cpp
@@ -162,7 +162,7 @@ FString FPipInstall::WritePluginsListing(TArray<TSharedRef<IPlugin>>& OutPythonP
 	FFileHelper::SaveStringArrayToFile(PipPluginPaths, *PyPluginsListingFile);
 
 	// Create .pth file in PipInstall/Lib/site-packages to account for plugins with packaged dependencies
-	const FString PyPluginsSitePackageFile = PipInstallPath / TEXT("Lib") / TEXT("site-packages") / PluginsSitePackageFilename;
+	const FString PyPluginsSitePackageFile = GetPipSitePackagesPath() / PluginsSitePackageFilename;
 	FFileHelper::SaveStringArrayToFile(PluginSitePackagePaths, *PyPluginsSitePackageFile);
 
     return PyPluginsListingFile;
