Added the ability to wrap a Python object for use as a UPROPERTY on a UCLASS or USTRUCT.

PythonObjectHandle allows Python generated types to store arbitrary Python objects as native properties, rather than rely on the lifetime of the Python wrapper object to keep a non-property object reference alive.

The following functions are available in Python script to interact with these handles:
  - unreal.create_python_object_handle(obj) -> handle  - Creates an unreal.PythonObjectHandle from a given PyObject.
  - unreal.resolve_python_object_handle(handle) -> obj - Resolves an unreal.PythonObjectHandle to its PyObject, or None.
  - unreal.destroy_python_object_handle(handle)        - Destroys an unreal.PythonObjectHandle (clearing its PyObject reference).

Manually calling destroy is optional, but can be useful when you need to *immediately* clean-up a reference to the wrapped Python object, rather than wait for UE GC to clean it up at an arbitrary point in the future.

This can be used to create a property on Python generated types via the standard unreal.uproperty syntax, eg):
  MyObjectHandleProperty = unreal.uproperty(unreal.PythonObjectHandle)

diff --git a/Engine/Plugins/Experimental/PythonScriptPlugin/Source/PythonScriptPlugin/Private/PyCore.cpp b/Engine/Plugins/Experimental/PythonScriptPlugin/Source/PythonScriptPlugin/Private/PyCore.cpp
--- a/Engine/Plugins/Experimental/PythonScriptPlugin/Source/PythonScriptPlugin/Private/PyCore.cpp
+++ b/Engine/Plugins/Experimental/PythonScriptPlugin/Source/PythonScriptPlugin/Private/PyCore.cpp
@@ -1758,6 +1758,68 @@ PyObject* GetInterpreterExecutablePath()
 	return PyConversion::Pythonize(PythonPath);
 }
 
+PyObject* CreatePythonObjectHandle(PyObject* InSelf, PyObject* InArgs)
+{
+	PyObject* PyObj = nullptr;
+	if (!PyArg_ParseTuple(InArgs, "O:create_python_object_handle", &PyObj))
+	{
+		return nullptr;
+	}
+
+	UPythonObjectHandle* Handle = UPythonObjectHandle::Create(PyObj);
+	return PyConversion::Pythonize(Handle);
+}
+
+PyObject* ResolvePythonObjectHandle(PyObject* InSelf, PyObject* InArgs)
+{
+	PyObject* PyHandle = nullptr;
+	if (!PyArg_ParseTuple(InArgs, "O:resolve_python_object_handle", &PyHandle))
+	{
+		return nullptr;
+	}
+	check(PyHandle);
+
+	UPythonObjectHandle* Handle = nullptr;
+	if (!PyConversion::Nativize(PyHandle, Handle))
+	{
+		PyUtil::SetPythonError(PyExc_TypeError, TEXT("resolve_python_object_handle"), *FString::Printf(TEXT("Parameter must be a 'PythonObjectHandle' not '%s'"), *PyUtil::GetFriendlyTypename(PyHandle)));
+		return nullptr;
+	}
+	
+	if (Handle)
+	{
+		PyObject* PyObj = Handle->Resolve();
+		Py_INCREF(PyObj);
+		return PyObj;
+	}
+
+	Py_RETURN_NONE;
+}
+
+PyObject* DestroyPythonObjectHandle(PyObject* InSelf, PyObject* InArgs)
+{
+	PyObject* PyHandle = nullptr;
+	if (!PyArg_ParseTuple(InArgs, "O:destroy_python_object_handle", &PyHandle))
+	{
+		return nullptr;
+	}
+	check(PyHandle);
+
+	UPythonObjectHandle* Handle = nullptr;
+	if (!PyConversion::Nativize(PyHandle, Handle))
+	{
+		PyUtil::SetPythonError(PyExc_TypeError, TEXT("destroy_python_object_handle"), *FString::Printf(TEXT("Parameter must be a 'PythonObjectHandle' not '%s'"), *PyUtil::GetFriendlyTypename(PyHandle)));
+		return nullptr;
+	}
+
+	if (Handle)
+	{
+		Handle->ReleasePythonResources();
+	}
+
+	Py_RETURN_NONE;
+}
+
 PyObject* RegisterPythonShutdownCallback(PyObject* InSelf, PyObject* InArgs)
 {
 	PyObject* PyObj = nullptr;
@@ -1836,6 +1898,9 @@ PyMethodDef PyCoreMethods[] = {
 	{ "LOCTABLE", PyCFunctionCast(&CreateLocalizedTextFromStringTable), METH_VARARGS, "LOCTABLE(id: Union[Name, str], key: str) -> Text -- get a localized Text from the given string table id and key" },
 	{ "is_editor", PyCFunctionCast(&IsEditor), METH_NOARGS, "is_editor() -> bool -- tells if the editor is running or not" },
 	{ "get_interpreter_executable_path", PyCFunctionCast(&GetInterpreterExecutablePath), METH_NOARGS, "get_interpreter_executable_path() -> str -- get the path to the Python interpreter executable of the Python SDK this plugin was compiled against" },
+	{ "create_python_object_handle", PyCFunctionCast(&CreatePythonObjectHandle), METH_VARARGS, "create_python_object_handle(obj: Optional[Any]) -> Optional[PythonObjectHandle] -- create a PythonObjectHandle from the given PyObject" },
+	{ "resolve_python_object_handle", PyCFunctionCast(&ResolvePythonObjectHandle), METH_VARARGS, "resolve_python_object_handle(handle: Optional[PythonObjectHandle]) -> Optional[Any] -- resolve a PythonObjectHandle to its PyObject, or None" },
+	{ "destroy_python_object_handle", PyCFunctionCast(&DestroyPythonObjectHandle), METH_VARARGS, "destroy_python_object_handle(handle: Optional[PythonObjectHandle]) -- destroy a PythonObjectHandle (clearing its PyObject reference)" },
 	{ nullptr, nullptr, 0, nullptr }
 };
 
diff --git a/Engine/Plugins/Experimental/PythonScriptPlugin/Source/PythonScriptPlugin/Private/PyReferenceCollector.cpp b/Engine/Plugins/Experimental/PythonScriptPlugin/Source/PythonScriptPlugin/Private/PyReferenceCollector.cpp
--- a/Engine/Plugins/Experimental/PythonScriptPlugin/Source/PythonScriptPlugin/Private/PyReferenceCollector.cpp
+++ b/Engine/Plugins/Experimental/PythonScriptPlugin/Source/PythonScriptPlugin/Private/PyReferenceCollector.cpp
@@ -159,6 +159,15 @@ void FPyReferenceCollector::PurgeUnrealGeneratedTypes()
 		}
 	}
 
+	// Clean-up Python object handles
+	// The handles are instances of UPythonObjectHandle
+	{
+		ForEachObjectOfClass(UPythonObjectHandle::StaticClass(), [&FlagObjectForPurge](UObject* InObject)
+		{
+			FlagObjectForPurge(InObject, /*bMarkPendingKill*/false);
+		}, false, RF_ClassDefaultObject, EInternalObjectFlags::Native);
+	}
+
 	if (PurgingReferenceCollector.HasObjectToPurge())
 	{
 		Py_BEGIN_ALLOW_THREADS
diff --git a/Engine/Plugins/Experimental/PythonScriptPlugin/Source/PythonScriptPlugin/Private/PyWrapperBase.cpp b/Engine/Plugins/Experimental/PythonScriptPlugin/Source/PythonScriptPlugin/Private/PyWrapperBase.cpp
--- a/Engine/Plugins/Experimental/PythonScriptPlugin/Source/PythonScriptPlugin/Private/PyWrapperBase.cpp
+++ b/Engine/Plugins/Experimental/PythonScriptPlugin/Source/PythonScriptPlugin/Private/PyWrapperBase.cpp
@@ -157,3 +157,48 @@ UPythonResourceOwner::UPythonResourceOwner(const FObjectInitializer& ObjectIniti
 {
 }
 
+#if WITH_PYTHON
+
+UPythonObjectHandle* UPythonObjectHandle::Create(PyObject* PyObjPtr)
+{
+	UPythonObjectHandle* Handle = nullptr;
+	if (PyObjPtr && PyObjPtr != Py_None)
+	{
+		Handle = NewObject<UPythonObjectHandle>();
+		{
+			FPyScopedGIL GIL;
+			Handle->PyObj = FPyObjectPtr::NewReference(PyObjPtr);
+		}
+	}
+	return Handle;
+}
+
+PyObject* UPythonObjectHandle::Resolve() const
+{
+	FPyScopedGIL GIL;
+	PyObject* PyObjPtr = PyObj.GetPtr();
+	return PyObjPtr ? PyObjPtr : Py_None;
+}
+
+void UPythonObjectHandle::BeginDestroy()
+{
+	ReleasePythonResources();
+	Super::BeginDestroy();
+}
+
+void UPythonObjectHandle::ReleasePythonResources()
+{
+	// This may be called after Python has already shut down
+	if (Py_IsInitialized())
+	{
+		FPyScopedGIL GIL;
+		PyObj.Reset();
+	}
+	else
+	{
+		// Release ownership if Python has been shut down to avoid attempting to delete the objects (which are already dead)
+		PyObj.Release();
+	}
+}
+
+#endif // WITH_PYTHON
diff --git a/Engine/Plugins/Experimental/PythonScriptPlugin/Source/PythonScriptPlugin/Private/PyWrapperBase.h b/Engine/Plugins/Experimental/PythonScriptPlugin/Source/PythonScriptPlugin/Private/PyWrapperBase.h
--- a/Engine/Plugins/Experimental/PythonScriptPlugin/Source/PythonScriptPlugin/Private/PyWrapperBase.h
+++ b/Engine/Plugins/Experimental/PythonScriptPlugin/Source/PythonScriptPlugin/Private/PyWrapperBase.h
@@ -105,3 +105,59 @@ public:
 	 */
 	virtual void ReleasePythonResources() = 0;
 };
+
+/**
+ * Handle object to wrap a Python object for use as a UPROPERTY on a UCLASS or USTRUCT.
+ * This allows Python generated types to store arbitrary Python objects as native properties, 
+ * rather than rely on the lifetime of the Python wrapper object to keep a non-property object reference alive.
+ * 
+ * The following functions are available in Python script to interact with these handles:
+ *   - unreal.create_python_object_handle(obj) -> handle  - Creates an unreal.PythonObjectHandle from a given PyObject.
+ *   - unreal.resolve_python_object_handle(handle) -> obj - Resolves an unreal.PythonObjectHandle to its PyObject, or None.
+ *   - unreal.destroy_python_object_handle(handle)        - Destroys an unreal.PythonObjectHandle (clearing its PyObject reference).
+ * 
+ * Manually calling destroy is optional, but can be useful when you need to *immediately* clean-up a reference to the wrapped 
+ * Python object, rather than wait for UE GC to clean it up at an arbitrary point in the future.
+ * 
+ * This can be used to create a property on Python generated types via the standard unreal.uproperty syntax, eg):
+ *   MyObjectHandleProperty = unreal.uproperty(unreal.PythonObjectHandle)
+ */
+UCLASS(Transient, BlueprintType)
+class UPythonObjectHandle final : public UObject, public IPythonResourceOwner
+{
+	GENERATED_BODY()
+
+#if WITH_PYTHON
+
+public:
+	/**
+	 * Create a handle that wraps the given Python object.
+	 * @note Will return null if passed a null or None Python object.
+	 * @param PyObjPtr Python object to wrap, may be null or None.
+	 */
+	static UPythonObjectHandle* Create(PyObject* PyObjPtr);
+
+	/**
+	 * Resolve this handle to its underlying Python object, if any.
+	 * @return The wrapped object, or None (but not null; borrowed reference).
+	 */
+	PyObject* Resolve() const;
+
+	//~ UObject interface
+	virtual void BeginDestroy() override;
+
+	//~ IPythonResourceOwner interface
+	virtual void ReleasePythonResources() override;
+
+private:
+	/** Wrapped Python object, if any */
+	FPyObjectPtr PyObj;
+
+#else	// WITH_PYTHON
+
+public:
+	//~ IPythonResourceOwner interface
+	virtual void ReleasePythonResources() override {}
+
+#endif	// WITH_PYTHON
+};
