Don't use the EvenIfUnreachable variant when adding referenced objects as objects are no longer speculatively marked as unreachable during GC and doing so breaks post GC verification tests

diff --git a/Engine/Plugins/Experimental/PythonScriptPlugin/Source/PythonScriptPlugin/Private/PyReferenceCollector.cpp b/Engine/Plugins/Experimental/PythonScriptPlugin/Source/PythonScriptPlugin/Private/PyReferenceCollector.cpp
--- a/Engine/Plugins/Experimental/PythonScriptPlugin/Source/PythonScriptPlugin/Private/PyReferenceCollector.cpp
+++ b/Engine/Plugins/Experimental/PythonScriptPlugin/Source/PythonScriptPlugin/Private/PyReferenceCollector.cpp
@@ -196,8 +196,7 @@ void FPyReferenceCollector::PurgeUnrealGeneratedTypes()
 void FPyReferenceCollector::AddReferencedObjectsFromDelegate(FReferenceCollector& InCollector, FScriptDelegate& InDelegate)
 {
 	// Keep the delegate object alive if it's using a Python proxy instance
-	// We have to use the EvenIfUnreachable variant here as the objects are speculatively marked as unreachable during GC
-	if (UPythonCallableForDelegate* PythonCallableForDelegate = Cast<UPythonCallableForDelegate>(InDelegate.GetUObjectEvenIfUnreachable()))
+	if (UPythonCallableForDelegate* PythonCallableForDelegate = Cast<UPythonCallableForDelegate>(InDelegate.GetUObject()))
 	{
 		TWeakObjectPtr<UPythonCallableForDelegate> Ptr{PythonCallableForDelegate};
 		InCollector.AddReferencedObject(Ptr);
