PipInstall: Handle UE_PIPINSTALL_PATH env var overloading for pip installer path. Code cleanup to use singleton instance and handle UE_PYTHONPATH and additional paths settings as well.

#changelist validated
#virtualized

diff --git a/Engine/Plugins/Experimental/PythonScriptPlugin/Source/PythonScriptPlugin/Private/PipInstall.cpp b/Engine/Plugins/Experimental/PythonScriptPlugin/Source/PythonScriptPlugin/Private/PipInstall.cpp
--- a/Engine/Plugins/Experimental/PythonScriptPlugin/Source/PythonScriptPlugin/Private/PipInstall.cpp
+++ b/Engine/Plugins/Experimental/PythonScriptPlugin/Source/PythonScriptPlugin/Private/PipInstall.cpp
@@ -111,43 +111,33 @@ const FString FPipInstall::ExtraUrlsFilename = TEXT("extra_urls.txt");
 const FString FPipInstall::ParsedRequirementsFilename = TEXT("merged_requirements.txt");
 
 
-bool FPipInstall::EnabledOnStartup()
+FPipInstall& FPipInstall::Get()
 {
-	bool bRunOnStartup = GetDefault<UPythonScriptPluginSettings>()->bRunPipInstallOnStartup;
-	bool bCmdLineDisable = FParse::Param(FCommandLine::Get(), TEXT("DisablePipInstall"));
+	static FPipInstall Instance;
+	return Instance;
+}
+
 
+bool FPipInstall::IsEnabled()
+{
 	return bRunOnStartup && !bCmdLineDisable;
 }
 
+bool FPipInstall::IsCmdLineDisabled()
+{
+	return bCmdLineDisable;
+}
+
 FString FPipInstall::WritePluginsListing(TArray<TSharedRef<IPlugin>>& OutPythonPlugins)
 {
 	TRACE_CPUPROFILER_EVENT_SCOPE(FPipInstall::WritePluginsListing);
 
-	const FString PipInstallPath = GetPipInstallPath();
-
 	OutPythonPlugins.Empty();
 
 	// List of plugins with pip dependencies
 	TArray<FString> PipPluginPaths;
-	// List of enabled plugins' site-packages folders
-	TArray<FString> PluginSitePackagePaths;
 	for ( const TSharedRef<IPlugin>& Plugin : IPluginManager::Get().GetEnabledPlugins() )
 	{
-		const FString PythonContentPath = FPaths::ConvertRelativePathToFull(Plugin->GetContentDir() / TEXT("Python"));
-		const FString PluginPlatformSitePackagesPath = PythonContentPath / TEXT("Lib") / FPlatformMisc::GetUBTPlatform() / TEXT("site-packages");
-		const FString PluginGeneralSitePackagesPath = PythonContentPath / TEXT("Lib") / TEXT("site-packages");
-
-		// Write platform/general site-packages paths per-plugin to .pth file to account for packaged python dependencies during pip install
-		if (FPaths::DirectoryExists(PluginPlatformSitePackagesPath))
-		{
-			PluginSitePackagePaths.Add(PluginPlatformSitePackagesPath);
-		}
-
-		if (FPaths::DirectoryExists(PluginGeneralSitePackagesPath))
-		{
-			PluginSitePackagePaths.Add(PluginGeneralSitePackagesPath);
-		}
-
 		const FPluginDescriptor& PluginDesc = Plugin->GetDescriptor();
 		if (PluginDesc.CachedJson->HasTypedField(TEXT("PythonRequirements"), EJson::Array))
 		{
@@ -161,9 +151,8 @@ FString FPipInstall::WritePluginsListing(TArray<TSharedRef<IPlugin>>& OutPythonP
 	const FString PyPluginsListingFile = PipInstallPath / PluginsListingFilename;
 	FFileHelper::SaveStringArrayToFile(PipPluginPaths, *PyPluginsListingFile);
 
-	// Create .pth file in PipInstall/Lib/site-packages to account for plugins with packaged dependencies
-	const FString PyPluginsSitePackageFile = GetPipSitePackagesPath() / PluginsSitePackageFilename;
-	FFileHelper::SaveStringArrayToFile(PluginSitePackagePaths, *PyPluginsSitePackageFile);
+	// Create .pth file in site-packages dir to account for plugins with packaged dependencies
+	WriteSitePackagePthFile();
 
     return PyPluginsListingFile;
 }
@@ -172,8 +161,6 @@ FString FPipInstall::WritePluginDependencies(const TArray<TSharedRef<IPlugin>>&
 {
 	TRACE_CPUPROFILER_EVENT_SCOPE(FPipInstall::WritePluginDependencies);
 
-	const FString PipInstallPath = GetPipInstallPath();
-
 	OutRequirements.Empty();
 	OutExtraUrls.Empty();
 
@@ -309,7 +296,6 @@ void FPipInstall::CheckInvalidPipEnv()
 {
 	TRACE_CPUPROFILER_EVENT_SCOPE(FPipInstall::CheckInvalidPipEnv);
 
-	const FString PipInstallPath = GetPipInstallPath();
 	if (!FPaths::DirectoryExists(PipInstallPath))
 	{
 		return;
@@ -322,7 +308,7 @@ void FPipInstall::CheckInvalidPipEnv()
 		return;
 	}
 
-	const FString VenvVersion = ParseVenvVersion(PipInstallPath);
+	const FString VenvVersion = ParseVenvVersion();
 	if (VenvVersion == TEXT(PY_VERSION))
 	{
 		return;
@@ -338,12 +324,9 @@ void FPipInstall::SetupPipEnv(FFeedbackContext* Context, bool bForceRebuild /* =
 {
 	TRACE_CPUPROFILER_EVENT_SCOPE(FPipInstall::SetupPipEnv);
 
-	const FString PipInstallPath = GetPipInstallPath();
-	const FString VenvInterp = GetVenvInterpreter(PipInstallPath);
-
 	if (!bForceRebuild && FPaths::FileExists(VenvInterp))
 	{
-		SetupPipInstallUtils(VenvInterp, Context);
+		SetupPipInstallUtils(Context);
 		return;
 	}
 
@@ -363,16 +346,14 @@ void FPipInstall::SetupPipEnv(FFeedbackContext* Context, bool bForceRebuild /* =
 		return;
 	}
 
-	SetupPipInstallUtils(VenvInterp, Context);
+	SetupPipInstallUtils(Context);
 }
 
 void FPipInstall::RemoveParsedDependencyFiles()
 {
 	TRACE_CPUPROFILER_EVENT_SCOPE(FPipInstall::ParsePluginDependencies);
 
-	const FString PipInstallPath = GetPipInstallPath();
 	const FString ParsedReqsFile = PipInstallPath / ParsedRequirementsFilename;
-
 	if (FPaths::FileExists(ParsedReqsFile))
 	{
 		IPlatformFile& PlatformFile = FPlatformFileManager::Get().GetPlatformFile();
@@ -384,9 +365,6 @@ FString FPipInstall::ParsePluginDependencies(const FString& MergedInRequirements
 {
 	TRACE_CPUPROFILER_EVENT_SCOPE(FPipInstall::ParsePluginDependencies);
 
-	const FString PipInstallPath = GetPipInstallPath();
-	const FString VenvInterp = GetVenvInterpreter(PipInstallPath);
-
 	const FString ParsedReqsFile = PipInstallPath / ParsedRequirementsFilename;
 
 	const FString Cmd = FString::Printf(TEXT("-m ue_parse_plugin_reqs -vv \"%s\" \"%s\""), *MergedInRequirementsFile, *ParsedReqsFile);
@@ -399,9 +377,6 @@ bool FPipInstall::RunPipInstall(FFeedbackContext* Context, bool bOfflineOnly, co
 {
 	TRACE_CPUPROFILER_EVENT_SCOPE(FPipInstall::RunPipInstall);
 
-	const FString PipInstallPath = GetPipInstallPath();
-	const FString VenvInterp = GetVenvInterpreter(PipInstallPath);
-
 	const FString ParsedReqsFile = PipInstallPath / ParsedRequirementsFilename;
 	const FString ExtraUrlsFile = PipInstallPath / ExtraUrlsFilename;
 
@@ -431,7 +406,7 @@ bool FPipInstall::RunPipInstall(FFeedbackContext* Context, bool bOfflineOnly, co
 	}
 	else if (!ForceIndexUrl.IsEmpty())
 	{
-		Cmd += TEXT("--index-url ") + ForceIndexUrl;
+		Cmd += TEXT(" --index-url ") + ForceIndexUrl;
 	}
 	else if (!ExtraUrls.IsEmpty())
 	{
@@ -448,6 +423,19 @@ bool FPipInstall::RunPipInstall(FFeedbackContext* Context, bool bOfflineOnly, co
 	return (Result == 0);
 }
 
+int FPipInstall::NumPackagesToInstall()
+{
+	const FString ParsedReqsFile = FPaths::ConvertRelativePathToFull(PipInstallPath / ParsedRequirementsFilename);
+
+	TArray<FString> ParsedReqLines;
+	if (!FPaths::FileExists(ParsedReqsFile) || !FFileHelper::LoadFileToStringArray(ParsedReqLines, *ParsedReqsFile))
+	{
+		return 0;
+	}
+
+	return FPipInstall::CountInstallLines(ParsedReqLines);
+}
+
 int FPipInstall::CountInstallLines(const TArray<FString>& RequirementLines)
 {
 	int Count = 0;
@@ -465,7 +453,7 @@ int FPipInstall::CountInstallLines(const TArray<FString>& RequirementLines)
 
 FString FPipInstall::GetPipInstallPath()
 {
-	return FPaths::ProjectIntermediateDir() / TEXT("PipInstall");
+	return PipInstallPath;
 }
 
 FString FPipInstall::GetPipSitePackagesPath()
@@ -481,16 +469,98 @@ FString FPipInstall::GetPipSitePackagesPath()
 }
 
 
-void FPipInstall::SetupPipInstallUtils(const FString& VenvInterp, FFeedbackContext* Context)
+FPipInstall::FPipInstall()
+{
+	TRACE_CPUPROFILER_EVENT_SCOPE(FPipInstall::FPipInstall);
+
+	// Check settings/cmd-line for whether pip installer is enabled
+	bRunOnStartup = GetDefault<UPythonScriptPluginSettings>()->bRunPipInstallOnStartup;
+	bCmdLineDisable = FParse::Param(FCommandLine::Get(), TEXT("DisablePipInstall"));
+
+	// Default install path: <ProjectDir>/PipInstall
+	PipInstallPath = FPaths::ConvertRelativePathToFull(FPaths::ProjectIntermediateDir() / TEXT("PipInstall"));
+
+	// Check for UE_PIPINSTALL_PATH install path override
+	const FString EnvInstallPath = FPlatformMisc::GetEnvironmentVariable(TEXT("UE_PIPINSTALL_PATH"));
+	if (!EnvInstallPath.IsEmpty())
+	{
+		FText ErrReason;
+		if (FPaths::ValidatePath(EnvInstallPath, &ErrReason))
+		{
+			PipInstallPath = FPaths::ConvertRelativePathToFull(EnvInstallPath);
+		}
+		else
+		{
+			UE_LOG(LogPython, Warning, TEXT("UE_PIPINSTALL_PATH: Invalid path specified: %s"), *ErrReason.ToString());
+		}
+	}
+
+	VenvInterp = GetVenvInterpreter(PipInstallPath);
+}
+
+
+void FPipInstall::WriteSitePackagePthFile()
+{
+	// Write all paths from script-plugin
+	// TODO: Should we directly use PyUtil::GetSystemPaths instead?
+	// TArray<FString> PluginSitePackagePaths = PyUtil::GetSystemPaths();
+
+	// List of enabled plugins' site-packages folders
+	TArray<FString> PluginSitePackagePaths;
+	for ( const TSharedRef<IPlugin>& Plugin : IPluginManager::Get().GetEnabledPlugins() )
+	{
+		const FString PythonContentPath = FPaths::ConvertRelativePathToFull(Plugin->GetContentDir() / TEXT("Python"));
+		const FString PluginPlatformSitePackagesPath = PythonContentPath / TEXT("Lib") / FPlatformMisc::GetUBTPlatform() / TEXT("site-packages");
+		const FString PluginGeneralSitePackagesPath = PythonContentPath / TEXT("Lib") / TEXT("site-packages");
+
+		// Write platform/general site-packages paths per-plugin to .pth file to account for packaged python dependencies during pip install
+		if (FPaths::DirectoryExists(PluginPlatformSitePackagesPath))
+		{
+			PluginSitePackagePaths.Add(PluginPlatformSitePackagesPath);
+		}
+
+		if (FPaths::DirectoryExists(PluginGeneralSitePackagesPath))
+		{
+			PluginSitePackagePaths.Add(PluginGeneralSitePackagesPath);
+		}
+	}
+
+	// Additional paths
+	for (const FDirectoryPath& AdditionalPath : GetDefault<UPythonScriptPluginSettings>()->AdditionalPaths)
+	{
+		const FString AddPath = FPaths::ConvertRelativePathToFull(AdditionalPath.Path);
+		if (FPaths::DirectoryExists(AddPath))
+		{
+			PluginSitePackagePaths.Add(AddPath);
+		}
+	}
+
+	// UE_PYTHONPATH
+	TArray<FString> SystemEnvPaths;
+	FPlatformMisc::GetEnvironmentVariable(TEXT("UE_PYTHONPATH")).ParseIntoArray(SystemEnvPaths, FPlatformMisc::GetPathVarDelimiter());
+	for (const FString& SystemEnvPath : SystemEnvPaths)
+	{
+		if (FPaths::DirectoryExists(SystemEnvPath))
+		{
+			PluginSitePackagePaths.Add(SystemEnvPath);
+		}
+	}
+
+	// Create .pth file in PipInstall/Lib/site-packages to account for plugins with packaged dependencies
+	const FString PyPluginsSitePackageFile = GetPipSitePackagesPath() / PluginsSitePackageFilename;
+	FFileHelper::SaveStringArrayToFile(PluginSitePackagePaths, *PyPluginsSitePackageFile);
+}
+
+
+void FPipInstall::SetupPipInstallUtils(FFeedbackContext* Context)
 {
 	TRACE_CPUPROFILER_EVENT_SCOPE(FPipInstall::SetupPipInstallUtils);
 
-	if (CheckPipInstallUtils(VenvInterp, Context))
+	if (CheckPipInstallUtils(Context))
 	{
 		return;
 	}
 
-	const FString PipInstallPath = GetPipInstallPath();
 	const FString PythonScriptDir = GetPythonScriptPluginPath();
 	if (PythonScriptDir.IsEmpty())
 	{
@@ -507,7 +577,7 @@ void FPipInstall::SetupPipInstallUtils(const FString& VenvInterp, FFeedbackConte
 }
 
 
-bool FPipInstall::CheckPipInstallUtils(const FString& VenvInterp, FFeedbackContext* Context)
+bool FPipInstall::CheckPipInstallUtils(FFeedbackContext* Context)
 {
 	// Verify that correct version of pip install utils is already available
 	const FString Cmd = FString::Printf(TEXT("-c \"import pkg_resources;dist=pkg_resources.working_set.find(pkg_resources.Requirement.parse('ue-pipinstall-utils'));exit(dist.version!='%s' if dist is not None else 1)\""), *PipInstallUtilsVer);
@@ -518,7 +588,7 @@ int32 FPipInstall::RunPythonCmd(const FText& Description, const FString& PythonI
 {
 	TRACE_CPUPROFILER_EVENT_SCOPE(FPipInstall::RunPythonCmd);
 
-	UE_LOG(LogPython, Log, TEXT("Running python command: python %s"), *Cmd);
+	UE_LOG(LogPython, Log, TEXT("Running python command: \"%s\" %s"), *PythonInterp, *Cmd);
 
 	int32 Result = 0;
 	RunLoggedSubprocess(&Result, Description, FPaths::ConvertRelativePathToFull(PythonInterp), Cmd, Context, CmdParser);
@@ -593,9 +663,9 @@ FString FPipInstall::GetPythonScriptPluginPath()
 	return PythonPlugin->GetBaseDir();
 }
 
-FString FPipInstall::ParseVenvVersion(const FString& InstallPath)
+FString FPipInstall::ParseVenvVersion()
 {
-	FString VenvConfig = InstallPath / TEXT("pyvenv.cfg");
+	FString VenvConfig = PipInstallPath / TEXT("pyvenv.cfg");
 	if (!FPaths::FileExists(VenvConfig))
 	{
 		return TEXT("");
diff --git a/Engine/Plugins/Experimental/PythonScriptPlugin/Source/PythonScriptPlugin/Private/PipInstall.h b/Engine/Plugins/Experimental/PythonScriptPlugin/Source/PythonScriptPlugin/Private/PipInstall.h
--- a/Engine/Plugins/Experimental/PythonScriptPlugin/Source/PythonScriptPlugin/Private/PipInstall.h
+++ b/Engine/Plugins/Experimental/PythonScriptPlugin/Source/PythonScriptPlugin/Private/PipInstall.h
@@ -16,41 +16,56 @@ class IProgressParser;
 class FPipInstall
 {
 public:
-	static const FString PipInstallUtilsVer;
-	static const FString PluginsListingFilename;
-	static const FString PluginsSitePackageFilename;
-	static const FString RequirementsInputFilename;
-	static const FString ExtraUrlsFilename;
-	static const FString ParsedRequirementsFilename;
+	static FPipInstall& Get();
 
-	static bool EnabledOnStartup();
+	bool IsEnabled();
+	bool IsCmdLineDisabled();
 
-	static void CheckRemoveOrphanedPackages(const FString& SitePackagesPath);
-	static void CheckInvalidPipEnv();
+	void CheckRemoveOrphanedPackages(const FString& SitePackagesPath);
+	void CheckInvalidPipEnv();
 
-	static FString WritePluginsListing(TArray<TSharedRef<IPlugin>>& OutPythonPlugins);
-	static FString WritePluginDependencies(const TArray<TSharedRef<IPlugin>>& PythonPlugins, TArray<FString>& OutRequirements, TArray<FString>& OutExtraUrls);
+	FString WritePluginsListing(TArray<TSharedRef<IPlugin>>& OutPythonPlugins);
+	FString WritePluginDependencies(const TArray<TSharedRef<IPlugin>>& PythonPlugins, TArray<FString>& OutRequirements, TArray<FString>& OutExtraUrls);
 
-	static void SetupPipEnv(FFeedbackContext* Context, bool bForceRebuild = false);
-	static void RemoveParsedDependencyFiles();
-	static FString ParsePluginDependencies(const FString& MergedInRequirementsFile, FFeedbackContext* Context);
-	static bool RunPipInstall(FFeedbackContext* Context, bool bOfflineOnly = false, const FString& ForceIndexUrl = TEXT(""));
+	void SetupPipEnv(FFeedbackContext* Context, bool bForceRebuild = false);
+	void RemoveParsedDependencyFiles();
+	FString ParsePluginDependencies(const FString& MergedInRequirementsFile, FFeedbackContext* Context);
+	bool RunPipInstall(FFeedbackContext* Context, bool bOfflineOnly = false, const FString& ForceIndexUrl = TEXT(""));
 
-	static int CountInstallLines(const TArray<FString>& RequirementLines);
+	int NumPackagesToInstall();
 
-	static FString GetPipInstallPath();
-	static FString GetPipSitePackagesPath();
+	FString GetPipInstallPath();
+	FString GetPipSitePackagesPath();
 
 private:
-	static void SetupPipInstallUtils(const FString& VenvInterp, FFeedbackContext* Context);
-	static bool CheckPipInstallUtils(const FString& VenvInterp, FFeedbackContext* Context);
+	FPipInstall();
+
+	void WriteSitePackagePthFile();
+	void SetupPipInstallUtils(FFeedbackContext* Context);
+	bool CheckPipInstallUtils(FFeedbackContext* Context);
 	static int32 RunPythonCmd(const FText& Description, const FString& PythonInterp, const FString& Cmd, FFeedbackContext* Context, TSharedPtr<IProgressParser> CmdParser = nullptr);
 	static bool RunLoggedSubprocess(int32* OutExitCode, const FText& Description, const FString& URL, const FString& Params, FFeedbackContext* Context, TSharedPtr<IProgressParser> CmdParser);
 
+	FString ParseVenvVersion();
+
 	static FString GetPythonScriptPluginPath();
-	static FString ParseVenvVersion(const FString& InstallPath);
 	static FString GetVenvInterpreter(const FString& InstallPath);
 	static bool CheckCompatiblePlatform(const TSharedPtr<FJsonObject>& JsonObject, const FString& PlatformName);
+
+	int CountInstallLines(const TArray<FString>& RequirementLines);
+
+	bool bRunOnStartup;
+	bool bCmdLineDisable;
+
+	FString PipInstallPath;
+	FString VenvInterp;
+
+	static const FString PipInstallUtilsVer;
+	static const FString PluginsListingFilename;
+	static const FString PluginsSitePackageFilename;
+	static const FString RequirementsInputFilename;
+	static const FString ExtraUrlsFilename;
+	static const FString ParsedRequirementsFilename;
 };
 
 #endif //WITH_PYTHON
diff --git a/Engine/Plugins/Experimental/PythonScriptPlugin/Source/PythonScriptPlugin/Private/PythonScriptPlugin.cpp b/Engine/Plugins/Experimental/PythonScriptPlugin/Source/PythonScriptPlugin/Private/PythonScriptPlugin.cpp
--- a/Engine/Plugins/Experimental/PythonScriptPlugin/Source/PythonScriptPlugin/Private/PythonScriptPlugin.cpp
+++ b/Engine/Plugins/Experimental/PythonScriptPlugin/Source/PythonScriptPlugin/Private/PythonScriptPlugin.cpp
@@ -984,7 +984,7 @@ void FPythonScriptPlugin::InitializePython()
 #endif // WITH_EDITOR
 
 		// Add Pip UBT install path to site-packages if it exists
-		const FString PipSitePackagePath = FPipInstall::GetPipSitePackagesPath();
+		const FString PipSitePackagePath = FPipInstall::Get().GetPipSitePackagesPath();
 		if (FPaths::DirectoryExists(PipSitePackagePath))
 		{
 			PyUtil::AddSitePackagesPath(PipSitePackagePath);
@@ -1134,38 +1134,39 @@ void FPythonScriptPlugin::InitPipInstaller()
 	FFeedbackContext* Context = GWarn;
 	FScopedSlowTask PipInstallTask(0, LOCTEXT("PipInstall.RunInit", "Running Pip Init Tasks..."), true, *Context);
 
-	FPipInstall::CheckInvalidPipEnv();
-	const FString PipSitePackagePath = FPaths::ConvertRelativePathToFull(FPipInstall::GetPipSitePackagesPath());
+	FPipInstall& PipInstaller = FPipInstall::Get();
+	PipInstaller.CheckInvalidPipEnv();
 
 	// Generate the input listing files of plugins with python dependencies and the listing of all requirements (installed or not)
 	TArray<TSharedRef<IPlugin>> PythonPlugins;
-	FPipInstall::WritePluginsListing(PythonPlugins);
+	PipInstaller.WritePluginsListing(PythonPlugins);
 
 	for (const TSharedRef<IPlugin>& PyPlugin : PythonPlugins)
 	{
 		// Remove leftover __pycache__ folders from plugins that use pip, but previously used packaged dependencies
 		const FString LibDir = PyPlugin->GetContentDir() / TEXT("Python") / TEXT("Lib");
-		FPipInstall::CheckRemoveOrphanedPackages(LibDir / TEXT("site-packages"));
-		FPipInstall::CheckRemoveOrphanedPackages(LibDir / FPlatformMisc::GetUBTPlatform() / TEXT("site-packages"));
+		PipInstaller.CheckRemoveOrphanedPackages(LibDir / TEXT("site-packages"));
+		PipInstaller.CheckRemoveOrphanedPackages(LibDir / FPlatformMisc::GetUBTPlatform() / TEXT("site-packages"));
 	}
 
 	TArray<FString> ReqInLines;
 	TArray<FString> ExtraUrls;
-	const FString ReqsInFile = FPipInstall::WritePluginDependencies(PythonPlugins, ReqInLines, ExtraUrls);
+	const FString ReqsInFile = PipInstaller.WritePluginDependencies(PythonPlugins, ReqInLines, ExtraUrls);
 
-	bool EnabledOnStart = FPipInstall::EnabledOnStartup();
+	bool EnabledOnStart = PipInstaller.IsEnabled();
 	if (ReqInLines.IsEmpty())
 	{
 		UE_CLOG(EnabledOnStart, LogPython, Display, TEXT("No enabled plugins with python dependencies found, skipping"));
-		FPipInstall::RemoveParsedDependencyFiles();
+		// Remove older parsed dependency files if there's nothing to install
+		PipInstaller.RemoveParsedDependencyFiles();
 		return;
 	}
 
 	// Just return immediately with warning if some python dependencies exist and pip install is disabled
+	const FString PipSitePackagePath = FPaths::ConvertRelativePathToFull(PipInstaller.GetPipSitePackagesPath());
 	if (!EnabledOnStart)
 	{
-		bool bCmdLineDisable = FParse::Param(FCommandLine::Get(), TEXT("DisablePipInstall"));
-		if (bCmdLineDisable || GIsBuildMachine)
+		if (PipInstaller.IsCmdLineDisabled() || GIsBuildMachine)
 		{
 			// Don't warn if disabled on cmd-line or is build process
 			UE_LOG(LogPython, Display, TEXT("Enabled plugins have python dependencies, install manually to: %s"), *PipSitePackagePath);
@@ -1181,31 +1182,30 @@ void FPythonScriptPlugin::InitPipInstaller()
 
 	UE_LOG(LogPython, Display, TEXT("Preparing to install python dependencies into: %s"), *PipSitePackagePath);
 
-	FPipInstall::SetupPipEnv(Context);
-	FPipInstall::ParsePluginDependencies(ReqsInFile, Context);
+	PipInstaller.SetupPipEnv(Context);
+	PipInstaller.ParsePluginDependencies(ReqsInFile, Context);
 }
 
 void FPythonScriptPlugin::RunPipInstaller()
 {
 	TRACE_CPUPROFILER_EVENT_SCOPE(FPythonScriptPlugin::RunPipInstaller);
 
-	if (!FPipInstall::EnabledOnStartup())
+	FPipInstall& PipInstaller = FPipInstall::Get();
+	if (!PipInstaller.IsEnabled())
 	{
 		return;
 	}
 
-	const FString ParsedReqsFile = FPaths::ConvertRelativePathToFull(FPipInstall::GetPipInstallPath() / FPipInstall::ParsedRequirementsFilename);
-
-	TArray<FString> ParsedReqLines;
-	if (!FPaths::FileExists(ParsedReqsFile) || !FFileHelper::LoadFileToStringArray(ParsedReqLines, *ParsedReqsFile) || (FPipInstall::CountInstallLines(ParsedReqLines) == 0))
+	if (PipInstaller.NumPackagesToInstall() == 0)
 	{
+		// Nothing to install
 		return;
 	}
 
 	FFeedbackContext* Context = GWarn;
 
 	// Run install of all python dependencies for enabled plugins
-	if (!FPipInstall::RunPipInstall(Context))
+	if (!PipInstaller.RunPipInstall(Context))
 	{
 		UE_LOG(LogPython, Warning, TEXT("Unable to install plugin python dependencies"));
 		return;
diff --git a/Engine/Plugins/Experimental/PythonScriptPlugin/Source/PythonScriptPlugin/Private/PythonScriptPlugin.h b/Engine/Plugins/Experimental/PythonScriptPlugin/Source/PythonScriptPlugin/Private/PythonScriptPlugin.h
--- a/Engine/Plugins/Experimental/PythonScriptPlugin/Source/PythonScriptPlugin/Private/PythonScriptPlugin.h
+++ b/Engine/Plugins/Experimental/PythonScriptPlugin/Source/PythonScriptPlugin/Private/PythonScriptPlugin.h
@@ -166,8 +166,6 @@ private:
 
 	void RunPipInstaller();
 
-	bool RunUBTPipAction(const FString& Action, const FText& Description, FFeedbackContext* Context);
-
 	void RequestStubCodeGeneration();
 
 	void GenerateStubCode();
diff --git a/Engine/Source/Programs/UnrealBuildTool/Modes/PipInstallMode.cs b/Engine/Source/Programs/UnrealBuildTool/Modes/PipInstallMode.cs
--- a/Engine/Source/Programs/UnrealBuildTool/Modes/PipInstallMode.cs
+++ b/Engine/Source/Programs/UnrealBuildTool/Modes/PipInstallMode.cs
@@ -112,6 +112,13 @@ namespace UnrealBuildTool.Modes
 			DirectoryReference InstallDir = DirectoryReference.Combine(ProjectDir, "Intermediate", "PipInstall");
 			UnrealTargetPlatform Platform = TargetDescriptor.Platform;
 
+			// Let env variable override pip install path
+			DirectoryReference? EnvInstallPath = DirectoryReference.FromString(Environment.GetEnvironmentVariable("UE_PIPINSTALL_PATH"));
+			if (EnvInstallPath != null)
+			{
+				InstallDir = EnvInstallPath;
+			}
+
 			PipEnv Pip = new(InstallDir, Platform, Logger, ProgressWriter.bWriteMarkup);
 			if ((Action & (PipAction)ActionBits.GenReqs) != 0)
 			{
@@ -237,12 +244,19 @@ namespace UnrealBuildTool.Modes
 
 		public void WritePluginsListing(UEBuildTarget Target, ILogger Logger)
 		{
+			// TODO: The path listing file won't match the .pth generated in-engine.
+			// In particular the additional paths setting
+
 			FileReference PluginsListingFile = FileReference.Combine(InstallDir, PluginsListingFilename);
 
 			DirectoryReference PipSitePackagesPath = DirectoryReference.Combine(InstallDir, "Lib", "site-packages");
 			FileReference PyPluginsSitePackageFile = FileReference.Combine(PipSitePackagesPath, PluginsSitePackageFilename);
 
-			FileReference.Delete(PluginsListingFile);
+			if (FileReference.Exists(PluginsListingFile))
+			{
+				FileReference.Delete(PluginsListingFile);
+			}
+
 			if (Target.EnabledPlugins == null)
 			{
 				return;
@@ -280,6 +294,17 @@ namespace UnrealBuildTool.Modes
 				}
 			}
 
+			// Add UE_PYTHONPATH to .pth file
+			string? UEPythonPaths = Environment.GetEnvironmentVariable("UE_PYTHONPATH");
+			if (UEPythonPaths != null)
+			{
+				string[] EnvPaths = UEPythonPaths.Split(Path.PathSeparator, StringSplitOptions.RemoveEmptyEntries | StringSplitOptions.TrimEntries);
+				foreach (string EnvPath in EnvPaths)
+				{
+					PluginsSitePackages.Add(EnvPath);
+				}
+			}
+
 			if (!DirectoryReference.Exists(PipSitePackagesPath))
 			{
 				DirectoryReference.CreateDirectory(PipSitePackagesPath);
