Disable hitch detection while generating the Python stub file

#rnx

diff --git a/Engine/Plugins/Experimental/PythonScriptPlugin/Source/PythonScriptPlugin/Private/PyWrapperTypeRegistry.cpp b/Engine/Plugins/Experimental/PythonScriptPlugin/Source/PythonScriptPlugin/Private/PyWrapperTypeRegistry.cpp
--- a/Engine/Plugins/Experimental/PythonScriptPlugin/Source/PythonScriptPlugin/Private/PyWrapperTypeRegistry.cpp
+++ b/Engine/Plugins/Experimental/PythonScriptPlugin/Source/PythonScriptPlugin/Private/PyWrapperTypeRegistry.cpp
@@ -1,7 +1,6 @@
 // Copyright Epic Games, Inc. All Rights Reserved.
 
 #include "PyWrapperTypeRegistry.h"
-#include "Misc/PackageName.h"
 #include "PyWrapperOwnerContext.h"
 #include "PyWrapperObject.h"
 #include "PyWrapperStruct.h"
@@ -23,6 +22,8 @@
 #include "PythonScriptPluginSettings.h"
 #include "ProfilingDebugging/ScopedTimers.h"
 #include "ProfilingDebugging/CpuProfilerTrace.h"
+#include "HAL/ThreadHeartBeat.h"
+#include "Misc/PackageName.h"
 #include "Misc/Paths.h"
 #include "Misc/FileHelper.h"
 #include "Misc/StringBuilder.h"
@@ -2599,6 +2600,10 @@ void FPyWrapperTypeRegistry::GenerateStubCodeForWrappedTypes(const EPyOnlineDocs
 {
 	UE_LOG(LogPython, Display, TEXT("Generating Python API stub file..."));
 
+	// Suspend the hang and hitch heartbeats, as this is a long running task.
+	FSlowHeartBeatScope SuspendHeartBeat;
+	FDisableHitchDetectorScope SuspendGameThreadHitch;
+
 	FPyScopedGIL GIL;
 	FPyFileWriter PythonScript;
 
