Fixed destroyed Python generated types leaving a dangling PyTypeObject* in FPyWrapperTypeRegistry

diff --git a/Engine/Plugins/Experimental/PythonScriptPlugin/Source/PythonScriptPlugin/Private/PyWrapperEnum.cpp b/Engine/Plugins/Experimental/PythonScriptPlugin/Source/PythonScriptPlugin/Private/PyWrapperEnum.cpp
--- a/Engine/Plugins/Experimental/PythonScriptPlugin/Source/PythonScriptPlugin/Private/PyWrapperEnum.cpp
+++ b/Engine/Plugins/Experimental/PythonScriptPlugin/Source/PythonScriptPlugin/Private/PyWrapperEnum.cpp
@@ -892,6 +892,10 @@ void UPythonGeneratedEnum::ReleasePythonResources()
 	if (Py_IsInitialized())
 	{
 		FPyScopedGIL GIL;
+		if (PyType)
+		{
+			FPyWrapperTypeRegistry::Get().UnregisterWrappedEnumType(GetFName(), PyType, !HasAnyFlags(RF_NewerVersionExists));
+		}
 		PyType.Reset();
 	}
 	else
diff --git a/Engine/Plugins/Experimental/PythonScriptPlugin/Source/PythonScriptPlugin/Private/PyWrapperObject.cpp b/Engine/Plugins/Experimental/PythonScriptPlugin/Source/PythonScriptPlugin/Private/PyWrapperObject.cpp
--- a/Engine/Plugins/Experimental/PythonScriptPlugin/Source/PythonScriptPlugin/Private/PyWrapperObject.cpp
+++ b/Engine/Plugins/Experimental/PythonScriptPlugin/Source/PythonScriptPlugin/Private/PyWrapperObject.cpp
@@ -1978,6 +1978,10 @@ void UPythonGeneratedClass::ReleasePythonResources()
 	if (Py_IsInitialized())
 	{
 		FPyScopedGIL GIL;
+		if (PyType)
+		{
+			FPyWrapperTypeRegistry::Get().UnregisterWrappedClassType(GetFName(), PyType, !HasAnyFlags(RF_NewerVersionExists));
+		}
 		PyType.Reset();
 		PyPostInitFunction.Reset();
 		for (const TSharedPtr<PyGenUtil::FFunctionDef>& FunctionDef : FunctionDefs)
diff --git a/Engine/Plugins/Experimental/PythonScriptPlugin/Source/PythonScriptPlugin/Private/PyWrapperStruct.cpp b/Engine/Plugins/Experimental/PythonScriptPlugin/Source/PythonScriptPlugin/Private/PyWrapperStruct.cpp
--- a/Engine/Plugins/Experimental/PythonScriptPlugin/Source/PythonScriptPlugin/Private/PyWrapperStruct.cpp
+++ b/Engine/Plugins/Experimental/PythonScriptPlugin/Source/PythonScriptPlugin/Private/PyWrapperStruct.cpp
@@ -1668,6 +1668,10 @@ void UPythonGeneratedStruct::ReleasePythonResources()
 	if (Py_IsInitialized())
 	{
 		FPyScopedGIL GIL;
+		if (PyType)
+		{
+			FPyWrapperTypeRegistry::Get().UnregisterWrappedStructType(GetFName(), PyType, !HasAnyFlags(RF_NewerVersionExists));
+		}
 		PyType.Reset();
 		PyPostInitFunction.Reset();
 	}
diff --git a/Engine/Plugins/Experimental/PythonScriptPlugin/Source/PythonScriptPlugin/Private/PyWrapperTypeRegistry.cpp b/Engine/Plugins/Experimental/PythonScriptPlugin/Source/PythonScriptPlugin/Private/PyWrapperTypeRegistry.cpp
--- a/Engine/Plugins/Experimental/PythonScriptPlugin/Source/PythonScriptPlugin/Private/PyWrapperTypeRegistry.cpp
+++ b/Engine/Plugins/Experimental/PythonScriptPlugin/Source/PythonScriptPlugin/Private/PyWrapperTypeRegistry.cpp
@@ -1842,9 +1842,12 @@ void FPyWrapperTypeRegistry::RegisterWrappedClassType(const FName ClassName, PyT
 	PythonWrappedClasses.Add(ClassName, PyType);
 }
 
-void FPyWrapperTypeRegistry::UnregisterWrappedClassType(const FName ClassName, PyTypeObject* PyType)
+void FPyWrapperTypeRegistry::UnregisterWrappedClassType(const FName ClassName, PyTypeObject* PyType, const bool InUnregisterName)
 {
-	UnregisterPythonTypeName(UTF8_TO_TCHAR(PyType->tp_name), ClassName);
+	if (InUnregisterName)
+	{
+		UnregisterPythonTypeName(UTF8_TO_TCHAR(PyType->tp_name), ClassName);
+	}
 	PythonWrappedClasses.Remove(ClassName);
 }
 
@@ -2140,9 +2143,12 @@ void FPyWrapperTypeRegistry::RegisterWrappedStructType(const FName StructName, P
 	PythonWrappedStructs.Add(StructName, PyType);
 }
 
-void FPyWrapperTypeRegistry::UnregisterWrappedStructType(const FName StructName, PyTypeObject* PyType)
+void FPyWrapperTypeRegistry::UnregisterWrappedStructType(const FName StructName, PyTypeObject* PyType, const bool InUnregisterName)
 {
-	UnregisterPythonTypeName(UTF8_TO_TCHAR(PyType->tp_name), StructName);
+	if (InUnregisterName)
+	{
+		UnregisterPythonTypeName(UTF8_TO_TCHAR(PyType->tp_name), StructName);
+	}
 	PythonWrappedStructs.Remove(StructName);
 }
 
@@ -2320,9 +2326,12 @@ void FPyWrapperTypeRegistry::RegisterWrappedEnumType(const FName EnumName, PyTyp
 	PythonWrappedEnums.Add(EnumName, PyType);
 }
 
-void FPyWrapperTypeRegistry::UnregisterWrappedEnumType(const FName EnumName, PyTypeObject* PyType)
+void FPyWrapperTypeRegistry::UnregisterWrappedEnumType(const FName EnumName, PyTypeObject* PyType, const bool InUnregisterName)
 {
-	UnregisterPythonTypeName(UTF8_TO_TCHAR(PyType->tp_name), EnumName);
+	if (InUnregisterName)
+	{
+		UnregisterPythonTypeName(UTF8_TO_TCHAR(PyType->tp_name), EnumName);
+	}
 	PythonWrappedEnums.Remove(EnumName);
 }
 
@@ -2492,9 +2501,12 @@ void FPyWrapperTypeRegistry::RegisterWrappedDelegateType(const FName DelegateNam
 	PythonWrappedDelegates.Add(DelegateName, PyType);
 }
 
-void FPyWrapperTypeRegistry::UnregisterWrappedDelegateType(const FName DelegateName, PyTypeObject* PyType)
+void FPyWrapperTypeRegistry::UnregisterWrappedDelegateType(const FName DelegateName, PyTypeObject* PyType, const bool InUnregisterName)
 {
-	UnregisterPythonTypeName(UTF8_TO_TCHAR(PyType->tp_name), DelegateName);
+	if (InUnregisterName)
+	{
+		UnregisterPythonTypeName(UTF8_TO_TCHAR(PyType->tp_name), DelegateName);
+	}
 	PythonWrappedDelegates.Remove(DelegateName);
 }
 
diff --git a/Engine/Plugins/Experimental/PythonScriptPlugin/Source/PythonScriptPlugin/Private/PyWrapperTypeRegistry.h b/Engine/Plugins/Experimental/PythonScriptPlugin/Source/PythonScriptPlugin/Private/PyWrapperTypeRegistry.h
--- a/Engine/Plugins/Experimental/PythonScriptPlugin/Source/PythonScriptPlugin/Private/PyWrapperTypeRegistry.h
+++ b/Engine/Plugins/Experimental/PythonScriptPlugin/Source/PythonScriptPlugin/Private/PyWrapperTypeRegistry.h
@@ -423,7 +423,7 @@ public:
 	void RegisterWrappedClassType(const FName ClassName, PyTypeObject* PyType, const bool InDetectNameConflicts = true);
 
 	/** Unregister the wrapped type associated with the given class name */
-	void UnregisterWrappedClassType(const FName ClassName, PyTypeObject* PyType);
+	void UnregisterWrappedClassType(const FName ClassName, PyTypeObject* PyType, const bool InUnregisterName = true);
 
 	/** True if we have wrapped type for the exact given class */
 	bool HasWrappedClassType(const UClass* InClass) const;
@@ -438,7 +438,7 @@ public:
 	void RegisterWrappedStructType(const FName StructName, PyTypeObject* PyType, const bool InDetectNameConflicts = true);
 
 	/** Unregister the wrapped type associated with the given struct name */
-	void UnregisterWrappedStructType(const FName StructName, PyTypeObject* PyType);
+	void UnregisterWrappedStructType(const FName StructName, PyTypeObject* PyType, const bool InUnregisterName = true);
 
 	/** True if we have wrapped type for the exact given struct */
 	bool HasWrappedStructType(const UScriptStruct* InStruct) const;
@@ -453,7 +453,7 @@ public:
 	void RegisterWrappedEnumType(const FName EnumName, PyTypeObject* PyType, const bool InDetectNameConflicts = true);
 
 	/** Unregister the wrapped type associated with the given enum name */
-	void UnregisterWrappedEnumType(const FName EnumName, PyTypeObject* PyType);
+	void UnregisterWrappedEnumType(const FName EnumName, PyTypeObject* PyType, const bool InUnregisterName = true);
 
 	/** True if we have wrapped type for the exact given enum */
 	bool HasWrappedEnumType(const UEnum* InEnum) const;
@@ -468,7 +468,7 @@ public:
 	void RegisterWrappedDelegateType(const FName DelegateName, PyTypeObject* PyType, const bool InDetectNameConflicts = true);
 
 	/** Unregister the wrapped type associated with the given delegate name */
-	void UnregisterWrappedDelegateType(const FName DelegateName, PyTypeObject* PyType);
+	void UnregisterWrappedDelegateType(const FName DelegateName, PyTypeObject* PyType, const bool InUnregisterName = true);
 
 	/** True if we have wrapped type for the exact given delegate signature */
 	bool HasWrappedDelegateType(const UFunction* InDelegateSignature) const;
