Updated Python to store generated types by their path name

This mostly involved making the process of emitting deprecated type aliases more robust, as there were a few cases where legitimate types were being skipped due to bugs in the previous short name implementation, which manifested as warnings in the path name version. The main fix is that we no longer generate a deprecated proxy for a field that still exists in memory, as it's now assumed that this field will emit its own (potentially deprecated) glue.

#rnx

diff --git a/Engine/Plugins/Experimental/PythonScriptPlugin/Source/PythonScriptPlugin/Private/PyGenUtil.cpp b/Engine/Plugins/Experimental/PythonScriptPlugin/Source/PythonScriptPlugin/Private/PyGenUtil.cpp
--- a/Engine/Plugins/Experimental/PythonScriptPlugin/Source/PythonScriptPlugin/Private/PyGenUtil.cpp
+++ b/Engine/Plugins/Experimental/PythonScriptPlugin/Source/PythonScriptPlugin/Private/PyGenUtil.cpp
@@ -2598,51 +2598,6 @@ const UObject* GetAssetTypeRegistryType(const UObject* InObj)
 	return InObj;
 }
 
-FName GetAssetTypeRegistryName(const UObject* InObj)
-{
-	// Note: If this changes then the functions in FPythonScriptPlugin that deal with FAssetData also need updating!
-	return InObj->GetOutermost()->GetFName();
-}
-
-FName GetTypeRegistryName(const UClass* InClass)
-{
-	return IsBlueprintGeneratedClass(InClass)
-		? GetAssetTypeRegistryName(InClass)
-		: InClass->GetFName();
-}
-
-FName GetTypeRegistryName(const UScriptStruct* InStruct)
-{
-	return IsBlueprintGeneratedStruct(InStruct)
-		? GetAssetTypeRegistryName(InStruct)
-		: InStruct->GetFName();
-}
-
-FName GetTypeRegistryName(const UEnum* InEnum)
-{
-	return IsBlueprintGeneratedEnum(InEnum)
-		? GetAssetTypeRegistryName(InEnum)
-		: InEnum->GetFName();
-}
-
-FName GetTypeRegistryName(const UFunction* InDelegateSignature)
-{
-	// Two different UObject-based classes can each declare in their body a delegate with the same type name, but possibly
-	// with different parameters. While they will share the same FName, they aren't be the same type. Make the type registry
-	// name unique to each delegate.
-	if (UObject* Outer = InDelegateSignature->GetOuter())
-	{
-		if (UClass* OuterClass = Cast<UClass>(Outer))
-		{
-			FNameBuilder ConcatName(Outer->GetFName());
-			ConcatName += "_";
-			InDelegateSignature->GetFName().AppendString(ConcatName);
-			return FName(*ConcatName);
-		}
-	}
-	return InDelegateSignature->GetFName();
-}
-
 FString GetFieldModule(const UField* InField)
 {
 	UPackage* ScriptPackage = InField->GetOutermost();
@@ -2741,7 +2696,7 @@ bool GetFieldPythonNameFromMetaDataImpl(const FFieldVariant& InField, const FNam
 	return false;
 }
 
-bool GetDeprecatedFieldPythonNamesFromMetaDataImpl(const FFieldVariant& InField, const FName InMetaDataKey, TArray<FString>& OutFieldNames)
+bool GetDeprecatedFieldPythonNamesFromMetaDataImpl(const FFieldVariant& InField, const FName InMetaDataKey, TArray<TTuple<FSoftObjectPath, FString>>& OutFieldNames)
 {
 	// See if we have a name override in the meta-data
 	if (!InMetaDataKey.IsNone())
@@ -2751,28 +2706,38 @@ bool GetDeprecatedFieldPythonNamesFromMetaDataImpl(const FFieldVariant& InField,
 		// This may be a semi-colon separated list - everything but the first item is deprecated
 		if (!FieldName.IsEmpty())
 		{
-			FieldName.ParseIntoArray(OutFieldNames, TEXT(";"), false);
+			TArray<FString> LocalFieldNames;
+			FieldName.ParseIntoArray(LocalFieldNames, TEXT(";"), false);
 
 			// Remove the non-deprecated entry
-			if (OutFieldNames.Num() > 0)
+			if (LocalFieldNames.Num() > 0)
 			{
-				OutFieldNames.RemoveAt(0, EAllowShrinking::No);
+				LocalFieldNames.RemoveAt(0, EAllowShrinking::No);
 			}
 
 			// Trim whitespace and remove empty items
-			OutFieldNames.RemoveAll([](FString& InStr)
+			LocalFieldNames.RemoveAll([](FString& InStr)
 			{
 				InStr.TrimStartAndEndInline();
 				return InStr.IsEmpty();
 			});
 
-			for (const FString& FieldNamePart : OutFieldNames)
+			OutFieldNames.Reset(LocalFieldNames.Num());
 			{
-				FText ValidationError;
-				if (!IsValidName(FieldNamePart, &ValidationError))
+				const FString FieldLeafName = InField.GetName();
+				const FString FieldPathName = InField.GetPathName();
+				for (const FString& FieldNamePart : LocalFieldNames)
 				{
-					const FString FieldPathName = InField.IsUObject() ? InField.Get<UField>()->GetPathName() : InField.Get<FField>()->GetPathName();
-					REPORT_PYTHON_GENERATION_ISSUE(Error, TEXT("'%s' has an invalid '%s' meta-data value '%s' (from '%s'): %s."), *FieldPathName, *InMetaDataKey.ToString(), *FieldNamePart, *FieldName, *ValidationError.ToString());
+					FText ValidationError;
+					if (!IsValidName(FieldNamePart, &ValidationError))
+					{
+						REPORT_PYTHON_GENERATION_ISSUE(Error, TEXT("'%s' has an invalid '%s' meta-data value '%s' (from '%s'): %s."), *FieldPathName, *InMetaDataKey.ToString(), *FieldNamePart, *FieldName, *ValidationError.ToString());
+					}
+
+					FString DeprecatedFieldPathName = FieldPathName;
+					DeprecatedFieldPathName.RemoveFromEnd(FieldLeafName);
+					DeprecatedFieldPathName.Append(FieldNamePart);
+					OutFieldNames.Add(MakeTuple(FSoftObjectPath(DeprecatedFieldPathName), FieldNamePart));
 				}
 			}
 
@@ -2830,9 +2795,9 @@ FString GetFieldPythonNameImpl(const FFieldVariant& InField, const FName InMetaD
 	return FieldName;
 }
 
-TArray<FString> GetDeprecatedFieldPythonNamesImpl(const FFieldVariant& InField, const FName InMetaDataKey)
+TArray<TTuple<FSoftObjectPath, FString>> GetDeprecatedFieldPythonNamesImpl(const FFieldVariant& InField, const FName InMetaDataKey)
 {
-	TArray<FString> FieldNames;
+	TArray<TTuple<FSoftObjectPath, FString>> FieldNames;
 
 	// First see if we have a name override in the meta-data
 	if (GetDeprecatedFieldPythonNamesFromMetaDataImpl(InField, InMetaDataKey, FieldNames))
@@ -2862,7 +2827,7 @@ TArray<FString> GetDeprecatedFieldPythonNamesImpl(const FFieldVariant& InField,
 	{
 		RedirectFlags = ECoreRedirectFlags::Type_Enum;
 	}
-	
+
 	const FCoreRedirectObjectName CurrentName = FCoreRedirectObjectName(InField.GetPathName());
 	TArray<FCoreRedirectObjectName> PreviousNames;
 	FCoreRedirects::FindPreviousNames(RedirectFlags, CurrentName, PreviousNames);
@@ -2907,6 +2872,49 @@ TArray<FString> GetDeprecatedFieldPythonNamesImpl(const FFieldVariant& InField,
 			}
 		}
 
+		// Redirects may be from a live field
+		// We want to skip those redirects as the real field still exists
+		if (InField.IsA<UClass>() || InField.IsA<UScriptStruct>() || InField.IsA<UEnum>())
+		{
+			UObject* PreviousType = nullptr;
+			if (!PreviousName.PackageName.IsNone())
+			{
+				// Try a qualified search first...
+				PreviousType = StaticFindObject(InField.Get<UObject>()->GetClass(), nullptr, *PreviousName.ToString());
+			}
+			if (!PreviousType)
+			{
+				// Try an unqualified search too, as the redirector might be broken or incorrect...
+				PreviousType = StaticFindFirstObject(InField.Get<UObject>()->GetClass(), *PreviousName.ObjectName.ToString(), EFindFirstObjectOptions::ExactClass | EFindFirstObjectOptions::NativeFirst);
+			}
+			if (PreviousType)
+			{
+				continue;
+			}
+		}
+		else if (InField.IsA<UFunction>())
+		{
+			if (const UClass* OwnerClass = InField.Get<UFunction>()->GetOwnerClass())
+			{
+				if (const UFunction* ExistingFunction = OwnerClass->FindFunctionByName(PreviousName.ObjectName);
+					ExistingFunction && ShouldExportFunction(ExistingFunction))
+				{
+					continue;
+				}
+			}
+		}
+		else if (InField.IsA<FProperty>())
+		{
+			if (const UStruct* OwnerStruct = InField.Get<FProperty>()->GetOwnerStruct())
+			{
+				if (const FProperty* ExistingProperty = OwnerStruct->FindPropertyByName(PreviousName.ObjectName);
+					ExistingProperty && ShouldExportProperty(ExistingProperty))
+				{
+					continue;
+				}
+			}
+		}
+
 		FString FieldName = PreviousName.ObjectName.ToString();
 
 		// Strip the "E" prefix from enum names
@@ -2915,7 +2923,20 @@ TArray<FString> GetDeprecatedFieldPythonNamesImpl(const FFieldVariant& InField,
 			FieldName.RemoveAt(0, EAllowShrinking::No);
 		}
 
-		FieldNames.AddUnique(MoveTemp(FieldName));
+		// Only add a single redirector for each previous name
+		if (FieldNames.ContainsByPredicate([&FieldName](TTuple<FSoftObjectPath, FString>& FieldNamePair) { return FieldNamePair.Value == FieldName; }))
+		{
+			continue;
+		}
+
+		// If PackageName is empty then assume that we're still using the same package as the current field
+		FCoreRedirectObjectName PreviousNameWithPackage = PreviousName;
+		if (PreviousNameWithPackage.PackageName.IsNone())
+		{
+			PreviousNameWithPackage.PackageName = InField.GetOutermost()->GetFName();
+		}
+
+		FieldNames.Add(MakeTuple(FSoftObjectPath(PreviousNameWithPackage.ToString()), MoveTemp(FieldName)));
 	}
 
 	return FieldNames;
@@ -2926,7 +2947,7 @@ FString GetClassPythonName(const UClass* InClass)
 	return GetFieldPythonNameImpl(InClass, ScriptNameMetaDataKey);
 }
 
-TArray<FString> GetDeprecatedClassPythonNames(const UClass* InClass)
+TArray<TTuple<FSoftObjectPath, FString>> GetDeprecatedClassPythonNames(const UClass* InClass)
 {
 	return GetDeprecatedFieldPythonNamesImpl(InClass, ScriptNameMetaDataKey);
 }
@@ -2936,7 +2957,7 @@ FString GetStructPythonName(const UScriptStruct* InStruct)
 	return GetFieldPythonNameImpl(InStruct, ScriptNameMetaDataKey);
 }
 
-TArray<FString> GetDeprecatedStructPythonNames(const UScriptStruct* InStruct)
+TArray<TTuple<FSoftObjectPath, FString>> GetDeprecatedStructPythonNames(const UScriptStruct* InStruct)
 {
 	return GetDeprecatedFieldPythonNamesImpl(InStruct, ScriptNameMetaDataKey);
 }
@@ -2946,7 +2967,7 @@ FString GetEnumPythonName(const UEnum* InEnum)
 	return GetFieldPythonNameImpl(InEnum, ScriptNameMetaDataKey);
 }
 
-TArray<FString> GetDeprecatedEnumPythonNames(const UEnum* InEnum)
+TArray<TTuple<FSoftObjectPath, FString>> GetDeprecatedEnumPythonNames(const UEnum* InEnum)
 {
 	return GetDeprecatedFieldPythonNamesImpl(InEnum, ScriptNameMetaDataKey);
 }
@@ -3018,15 +3039,15 @@ FString GetFunctionPythonName(const UFunction* InFunc)
 	return PythonizeName(FuncName, EPythonizeNameCase::Lower);
 }
 
-TArray<FString> GetDeprecatedFunctionPythonNames(const UFunction* InFunc)
+TArray<TTuple<FSoftObjectPath, FString>> GetDeprecatedFunctionPythonNames(const UFunction* InFunc)
 {
 	const UClass* FuncOwner = InFunc->GetOwnerClass();
 	check(FuncOwner);
 
-	TArray<FString> FuncNames = GetDeprecatedFieldPythonNamesImpl(InFunc, ScriptNameMetaDataKey);
+	TArray<TTuple<FSoftObjectPath, FString>> FuncNames = GetDeprecatedFieldPythonNamesImpl(InFunc, ScriptNameMetaDataKey);
 	for (auto FuncNamesIt = FuncNames.CreateIterator(); FuncNamesIt; ++FuncNamesIt)
 	{
-		FString& FuncName = *FuncNamesIt;
+		FString& FuncName = FuncNamesIt->Value;
 
 		// Remove any deprecated names that clash with an existing Python exposed function
 		const UFunction* DeprecatedFunc = FuncOwner->FindFunctionByName(*FuncName);
@@ -3052,13 +3073,14 @@ FString GetScriptMethodPythonName(const UFunction* InFunc)
 	return GetFunctionPythonName(InFunc);
 }
 
-TArray<FString> GetDeprecatedScriptMethodPythonNames(const UFunction* InFunc)
+TArray<TTuple<FSoftObjectPath, FString>> GetDeprecatedScriptMethodPythonNames(const UFunction* InFunc)
 {
-	TArray<FString> ScriptMethodNames;
+	TArray<TTuple<FSoftObjectPath, FString>> ScriptMethodNames;
 	if (GetDeprecatedFieldPythonNamesFromMetaDataImpl(InFunc, ScriptMethodMetaDataKey, ScriptMethodNames))
 	{
-		for (FString& ScriptMethodName : ScriptMethodNames)
+		for (TTuple<FSoftObjectPath, FString>& ScriptMethodNamePair : ScriptMethodNames)
 		{
+			FString& ScriptMethodName = ScriptMethodNamePair.Value;
 			ScriptMethodName = PythonizeName(ScriptMethodName, EPythonizeNameCase::Lower);
 		}
 		return ScriptMethodNames;
@@ -3067,7 +3089,7 @@ TArray<FString> GetDeprecatedScriptMethodPythonNames(const UFunction* InFunc)
 }
 
 FString GetScriptConstantPythonName(const UFunction* InFunc)
-	{
+{
 	FString ScriptConstantName;
 	if (!GetFieldPythonNameFromMetaDataImpl(InFunc, ScriptConstantMetaDataKey, ScriptConstantName))
 	{
@@ -3076,15 +3098,16 @@ FString GetScriptConstantPythonName(const UFunction* InFunc)
 	return PythonizeName(ScriptConstantName, EPythonizeNameCase::Upper);
 }
 
-TArray<FString> GetDeprecatedScriptConstantPythonNames(const UFunction* InFunc)
+TArray<TTuple<FSoftObjectPath, FString>> GetDeprecatedScriptConstantPythonNames(const UFunction* InFunc)
 {
-	TArray<FString> ScriptConstantNames;
+	TArray<TTuple<FSoftObjectPath, FString>> ScriptConstantNames;
 	if (!GetDeprecatedFieldPythonNamesFromMetaDataImpl(InFunc, ScriptConstantMetaDataKey, ScriptConstantNames))
 	{
 		ScriptConstantNames = GetDeprecatedFieldPythonNamesImpl(InFunc, ScriptNameMetaDataKey);
 	}
-	for (FString& ScriptConstantName : ScriptConstantNames)
+	for (TTuple<FSoftObjectPath, FString>& ScriptConstantNamePair : ScriptConstantNames)
 	{
+		FString& ScriptConstantName = ScriptConstantNamePair.Value;
 		ScriptConstantName = PythonizeName(ScriptConstantName, EPythonizeNameCase::Upper);
 	}
 	return ScriptConstantNames;
@@ -3096,15 +3119,15 @@ FString GetPropertyPythonName(const FProperty* InProp)
 	return PythonizePropertyName(PropName, EPythonizeNameCase::Lower);
 }
 
-TArray<FString> GetDeprecatedPropertyPythonNames(const FProperty* InProp)
+TArray<TTuple<FSoftObjectPath, FString>> GetDeprecatedPropertyPythonNames(const FProperty* InProp)
 {
 	const UStruct* PropOwner = InProp->GetOwnerStruct();
 	check(PropOwner);
 
-	TArray<FString> PropNames = GetDeprecatedFieldPythonNamesImpl(InProp, ScriptNameMetaDataKey);
+	TArray<TTuple<FSoftObjectPath, FString>> PropNames = GetDeprecatedFieldPythonNamesImpl(InProp, ScriptNameMetaDataKey);
 	for (auto PropNamesIt = PropNames.CreateIterator(); PropNamesIt; ++PropNamesIt)
 	{
-		FString& PropName = *PropNamesIt;
+		FString& PropName = PropNamesIt->Value;
 
 		// Remove any deprecated names that clash with an existing Python exposed property
 		const FProperty* DeprecatedProp = PropOwner->FindPropertyByName(*PropName);
diff --git a/Engine/Plugins/Experimental/PythonScriptPlugin/Source/PythonScriptPlugin/Private/PyGenUtil.h b/Engine/Plugins/Experimental/PythonScriptPlugin/Source/PythonScriptPlugin/Private/PyGenUtil.h
--- a/Engine/Plugins/Experimental/PythonScriptPlugin/Source/PythonScriptPlugin/Private/PyGenUtil.h
+++ b/Engine/Plugins/Experimental/PythonScriptPlugin/Source/PythonScriptPlugin/Private/PyGenUtil.h
@@ -1124,21 +1124,6 @@ namespace PyGenUtil
 	/** Get the type that should be used with the Python type registry for an asset (eg, a Blueprint asset should use its generated class) */
 	const UObject* GetAssetTypeRegistryType(const UObject* InObj);
 
-	/** Get the name that should be used with the Python type registry for an asset (this uses the package path name, as assets may not have unique names) */
-	FName GetAssetTypeRegistryName(const UObject* InObj);
-
-	/** Get the name that should be used by the given class when registered with the Python type registry */
-	FName GetTypeRegistryName(const UClass* InClass);
-
-	/** Get the name that should be used by the given struct when registered with the Python type registry */
-	FName GetTypeRegistryName(const UScriptStruct* InStruct);
-
-	/** Get the name that should be used by the given enum when registered with the Python type registry */
-	FName GetTypeRegistryName(const UEnum* InEnum);
-
-	/** Get the name that should be used by the given delegate when registered with the Python type registry */
-	FName GetTypeRegistryName(const UFunction* InDelegateSignature);
-
 	/** Get the native module the given field belongs to */
 	FString GetFieldModule(const UField* InField);
 
@@ -1152,19 +1137,19 @@ namespace PyGenUtil
 	FString GetClassPythonName(const UClass* InClass);
 
 	/** Get the deprecated Python names of the given class */
-	TArray<FString> GetDeprecatedClassPythonNames(const UClass* InClass);
+	TArray<TTuple<FSoftObjectPath, FString>> GetDeprecatedClassPythonNames(const UClass* InClass);
 
 	/** Get the Python name of the given struct */
 	FString GetStructPythonName(const UScriptStruct* InStruct);
 
 	/** Get the deprecated Python names of the given struct */
-	TArray<FString> GetDeprecatedStructPythonNames(const UScriptStruct* InStruct);
+	TArray<TTuple<FSoftObjectPath, FString>> GetDeprecatedStructPythonNames(const UScriptStruct* InStruct);
 
 	/** Get the Python name of the given enum */
 	FString GetEnumPythonName(const UEnum* InEnum);
 
 	/** Get the deprecated Python names of the given enum */
-	TArray<FString> GetDeprecatedEnumPythonNames(const UEnum* InEnum);
+	TArray<TTuple<FSoftObjectPath, FString>> GetDeprecatedEnumPythonNames(const UEnum* InEnum);
 
 	/** Get the Python name of the given enum entry */
 	FString GetEnumEntryPythonName(const UEnum* InEnum, const int32 InEntryIndex);
@@ -1176,25 +1161,25 @@ namespace PyGenUtil
 	FString GetFunctionPythonName(const UFunction* InFunc);
 
 	/** Get the deprecated Python names of the given function */
-	TArray<FString> GetDeprecatedFunctionPythonNames(const UFunction* InFunc);
+	TArray<TTuple<FSoftObjectPath, FString>> GetDeprecatedFunctionPythonNames(const UFunction* InFunc);
 
 	/** Get the Python name of the given function when it's hoisted as a script method */
 	FString GetScriptMethodPythonName(const UFunction* InFunc);
 
 	/** Get the deprecated Python names of the given function it's hoisted as a script method */
-	TArray<FString> GetDeprecatedScriptMethodPythonNames(const UFunction* InFunc);
+	TArray<TTuple<FSoftObjectPath, FString>> GetDeprecatedScriptMethodPythonNames(const UFunction* InFunc);
 
 	/** Get the Python name of the given function when it's hoisted as a script constant */
 	FString GetScriptConstantPythonName(const UFunction* InFunc);
 
 	/** Get the deprecated Python names of the given function it's hoisted as a script constant */
-	TArray<FString> GetDeprecatedScriptConstantPythonNames(const UFunction* InFunc);
+	TArray<TTuple<FSoftObjectPath, FString>> GetDeprecatedScriptConstantPythonNames(const UFunction* InFunc);
 
 	/** Get the Python name of the given property */
 	FString GetPropertyPythonName(const FProperty* InProp);
 
 	/** Get the deprecated Python names of the given property */
-	TArray<FString> GetDeprecatedPropertyPythonNames(const FProperty* InProp);
+	TArray<TTuple<FSoftObjectPath, FString>> GetDeprecatedPropertyPythonNames(const FProperty* InProp);
 
 	/** Get the Python name of the given property */
 	FString GetPropertyTypePythonName(const FProperty* InProp, const bool InIncludeUnrealNamespace = false, const bool InIsForDocString = true, const EPythonizeFlags PythonizeFlags = EPythonizeFlags::None);
diff --git a/Engine/Plugins/Experimental/PythonScriptPlugin/Source/PythonScriptPlugin/Private/PyWrapperEnum.cpp b/Engine/Plugins/Experimental/PythonScriptPlugin/Source/PythonScriptPlugin/Private/PyWrapperEnum.cpp
--- a/Engine/Plugins/Experimental/PythonScriptPlugin/Source/PythonScriptPlugin/Private/PyWrapperEnum.cpp
+++ b/Engine/Plugins/Experimental/PythonScriptPlugin/Source/PythonScriptPlugin/Private/PyWrapperEnum.cpp
@@ -798,7 +798,7 @@ public:
 
 		// Map the Unreal enum to the Python type
 		NewEnum->PyType = FPyTypeObjectPtr::NewReference(PyType);
-		FPyWrapperTypeRegistry::Get().RegisterWrappedEnumType(NewEnum->GetFName(), PyType, !bDidExist);
+		FPyWrapperTypeRegistry::Get().RegisterWrappedEnumType(NewEnum, PyType, !bDidExist);
 
 		// Null the NewEnum pointer so the destructor doesn't kill it
 		UPythonGeneratedEnum* FinalizedEnum = NewEnum;
@@ -894,7 +894,7 @@ void UPythonGeneratedEnum::ReleasePythonResources()
 		FPyScopedGIL GIL;
 		if (PyType)
 		{
-			FPyWrapperTypeRegistry::Get().UnregisterWrappedEnumType(GetFName(), PyType, !HasAnyFlags(RF_NewerVersionExists));
+			FPyWrapperTypeRegistry::Get().UnregisterWrappedEnumType(this, PyType, !HasAnyFlags(RF_NewerVersionExists));
 		}
 		PyType.Reset();
 	}
diff --git a/Engine/Plugins/Experimental/PythonScriptPlugin/Source/PythonScriptPlugin/Private/PyWrapperObject.cpp b/Engine/Plugins/Experimental/PythonScriptPlugin/Source/PythonScriptPlugin/Private/PyWrapperObject.cpp
--- a/Engine/Plugins/Experimental/PythonScriptPlugin/Source/PythonScriptPlugin/Private/PyWrapperObject.cpp
+++ b/Engine/Plugins/Experimental/PythonScriptPlugin/Source/PythonScriptPlugin/Private/PyWrapperObject.cpp
@@ -1399,7 +1399,7 @@ public:
 
 		// Map the Unreal class to the Python type
 		NewClass->PyType = FPyTypeObjectPtr::NewReference(PyType);
-		FPyWrapperTypeRegistry::Get().RegisterWrappedClassType(NewClass->GetFName() , PyType);
+		FPyWrapperTypeRegistry::Get().RegisterWrappedClassType(NewClass, PyType);
 
 		// Ensure the CDO exists
 		Py_BEGIN_ALLOW_THREADS
@@ -1935,8 +1935,8 @@ void UPythonGeneratedClass::PostRename(UObject* OldOuter, const FName OldName)
 
 	if (PyType)
 	{
-		FPyWrapperTypeRegistry::Get().UnregisterWrappedClassType(OldName, PyType);
-		FPyWrapperTypeRegistry::Get().RegisterWrappedClassType(GetFName(), PyType, !HasAnyFlags(RF_NewerVersionExists));
+		FPyWrapperTypeRegistry::Get().UnregisterWrappedClassType(FSoftObjectPath(OldOuter->GetFName(), OldName, FString()), PyType);
+		FPyWrapperTypeRegistry::Get().RegisterWrappedClassType(this, PyType, !HasAnyFlags(RF_NewerVersionExists));
 	}
 }
 
@@ -1980,7 +1980,7 @@ void UPythonGeneratedClass::ReleasePythonResources()
 		FPyScopedGIL GIL;
 		if (PyType)
 		{
-			FPyWrapperTypeRegistry::Get().UnregisterWrappedClassType(GetFName(), PyType, !HasAnyFlags(RF_NewerVersionExists));
+			FPyWrapperTypeRegistry::Get().UnregisterWrappedClassType(this, PyType, !HasAnyFlags(RF_NewerVersionExists));
 		}
 		PyType.Reset();
 		PyPostInitFunction.Reset();
diff --git a/Engine/Plugins/Experimental/PythonScriptPlugin/Source/PythonScriptPlugin/Private/PyWrapperStruct.cpp b/Engine/Plugins/Experimental/PythonScriptPlugin/Source/PythonScriptPlugin/Private/PyWrapperStruct.cpp
--- a/Engine/Plugins/Experimental/PythonScriptPlugin/Source/PythonScriptPlugin/Private/PyWrapperStruct.cpp
+++ b/Engine/Plugins/Experimental/PythonScriptPlugin/Source/PythonScriptPlugin/Private/PyWrapperStruct.cpp
@@ -41,7 +41,7 @@ const IPyWrapperStructAllocationPolicy* GetPyWrapperStructAllocationPolicy(UScri
 		}
 	};
 
-	if (const IPyWrapperInlineStructFactory* InlineStructFactory = FPyWrapperTypeRegistry::Get().GetInlineStructFactory(InStruct->GetFName()))
+	if (const IPyWrapperInlineStructFactory* InlineStructFactory = FPyWrapperTypeRegistry::Get().GetInlineStructFactory(InStruct->GetStructPathName()))
 	{
 		return InlineStructFactory->GetPythonObjectAllocationPolicy();
 	}
@@ -1521,7 +1521,7 @@ public:
 
 		// Map the Unreal struct to the Python type
 		NewStruct->PyType = FPyTypeObjectPtr::NewReference(PyType);
-		FPyWrapperTypeRegistry::Get().RegisterWrappedStructType(NewStruct->GetFName(), PyType);
+		FPyWrapperTypeRegistry::Get().RegisterWrappedStructType(NewStruct, PyType);
 
 		// Re-instance the old struct
 		if (OldStruct)
@@ -1620,8 +1620,8 @@ void UPythonGeneratedStruct::PostRename(UObject* OldOuter, const FName OldName)
 
 	if (PyType)
 	{
-		FPyWrapperTypeRegistry::Get().UnregisterWrappedStructType(OldName, PyType);
-		FPyWrapperTypeRegistry::Get().RegisterWrappedStructType(GetFName(), PyType, !HasAnyFlags(RF_NewerVersionExists));
+		FPyWrapperTypeRegistry::Get().UnregisterWrappedStructType(FSoftObjectPath(OldOuter->GetFName(), OldName, FString()), PyType);
+		FPyWrapperTypeRegistry::Get().RegisterWrappedStructType(this, PyType, !HasAnyFlags(RF_NewerVersionExists));
 	}
 }
 
@@ -1670,7 +1670,7 @@ void UPythonGeneratedStruct::ReleasePythonResources()
 		FPyScopedGIL GIL;
 		if (PyType)
 		{
-			FPyWrapperTypeRegistry::Get().UnregisterWrappedStructType(GetFName(), PyType, !HasAnyFlags(RF_NewerVersionExists));
+			FPyWrapperTypeRegistry::Get().UnregisterWrappedStructType(this, PyType, !HasAnyFlags(RF_NewerVersionExists));
 		}
 		PyType.Reset();
 		PyPostInitFunction.Reset();
diff --git a/Engine/Plugins/Experimental/PythonScriptPlugin/Source/PythonScriptPlugin/Private/PyWrapperStruct.h b/Engine/Plugins/Experimental/PythonScriptPlugin/Source/PythonScriptPlugin/Private/PyWrapperStruct.h
--- a/Engine/Plugins/Experimental/PythonScriptPlugin/Source/PythonScriptPlugin/Private/PyWrapperStruct.h
+++ b/Engine/Plugins/Experimental/PythonScriptPlugin/Source/PythonScriptPlugin/Private/PyWrapperStruct.h
@@ -253,7 +253,7 @@ class IPyWrapperInlineStructFactory
 {
 public:
 	/** Get the name of the Unreal struct this factory is for */
-	virtual FName GetStructName() const = 0;
+	virtual FTopLevelAssetPath GetStructName() const = 0;
 
 	/** Get the size of the Python object that should be constructed (in bytes) */
 	virtual int32 GetPythonObjectSizeBytes() const = 0;
@@ -281,9 +281,9 @@ private:
 	};
 
 public:
-	virtual FName GetStructName() const override
+	virtual FTopLevelAssetPath GetStructName() const override
 	{
-		return TBaseStructure<InlineType>::Get()->GetFName();
+		return TBaseStructure<InlineType>::Get()->GetStructPathName();
 	}
 
 	virtual int32 GetPythonObjectSizeBytes() const override
diff --git a/Engine/Plugins/Experimental/PythonScriptPlugin/Source/PythonScriptPlugin/Private/PyWrapperTypeRegistry.cpp b/Engine/Plugins/Experimental/PythonScriptPlugin/Source/PythonScriptPlugin/Private/PyWrapperTypeRegistry.cpp
--- a/Engine/Plugins/Experimental/PythonScriptPlugin/Source/PythonScriptPlugin/Private/PyWrapperTypeRegistry.cpp
+++ b/Engine/Plugins/Experimental/PythonScriptPlugin/Source/PythonScriptPlugin/Private/PyWrapperTypeRegistry.cpp
@@ -26,6 +26,7 @@
 #include "Misc/PackageName.h"
 #include "Misc/Paths.h"
 #include "Misc/FileHelper.h"
+#include "Misc/ScopedSlowTask.h"
 #include "Misc/StringBuilder.h"
 #include "SourceCodeNavigation.h"
 #include "UObject/UnrealType.h"
@@ -40,6 +41,8 @@
 #include "Kismet2/ReloadUtilities.h"
 #endif
 
+#define LOCTEXT_NAMESPACE "PyWrapperTypeRegistry"
+
 #if WITH_PYTHON
 
 namespace UE::Python
@@ -89,7 +92,7 @@ void SetMakeFunction(FPyWrapperStructMetaData& MetaData, UFunction* MakeFunc)
 
 	if (!bHasValidReturn)
 	{
-		REPORT_PYTHON_GENERATION_ISSUE(Warning, TEXT("Struct '%s' is marked as 'HasNativeMake' but the function '%s' does not return the struct type. Python will use the generic make function to create a Python object of this type."), *MetaData.Struct->GetName(), *MetaData.MakeFunc.Func->GetPathName());
+		REPORT_PYTHON_GENERATION_ISSUE(Warning, TEXT("Struct '%s' is marked as 'HasNativeMake' but the function '%s' does not return the struct type. Python will use the generic make function to create a Python object of this type."), *MetaData.Struct->GetPathName(), *MetaData.MakeFunc.Func->GetPathName());
 		MetaData.MakeFunc.SetFunction(nullptr);
 		return;
 	}
@@ -117,7 +120,7 @@ void SetBreakFunction(FPyWrapperStructMetaData& MetaData, UFunction* BreakFunc)
 
 	if (!bHasValidInput)
 	{
-		REPORT_PYTHON_GENERATION_ISSUE(Warning, TEXT("Struct '%s' is marked as 'HasNativeBreak' but the function '%s' does not have the struct type as its only input argument. Python will use the generic break function to convert a Python object of this type into a tuple."), *MetaData.Struct->GetName(), *MetaData.BreakFunc.Func->GetPathName());
+		REPORT_PYTHON_GENERATION_ISSUE(Warning, TEXT("Struct '%s' is marked as 'HasNativeBreak' but the function '%s' does not have the struct type as its only input argument. Python will use the generic break function to convert a Python object of this type into a tuple."), *MetaData.Struct->GetPathName(), *MetaData.BreakFunc.Func->GetPathName());
 		MetaData.BreakFunc.SetFunction(nullptr);
 	}
 };
@@ -820,11 +823,11 @@ FPyWrapperTypeRegistry::FPyWrapperTypeRegistry()
 	{
 		for (const TPair<FString, TSharedPtr<FPyWrapperStructMetaData>>& UnresolvedMakePair : UnresolvedMakeFuncs)
 		{
-			REPORT_PYTHON_GENERATION_ISSUE(Error, TEXT("Struct '%s' is marked as '%s' but the function '%s' could not be found."), *UnresolvedMakePair.Value->Struct->GetName(), *PyGenUtil::HasNativeMakeMetaDataKey.ToString(), *UnresolvedMakePair.Key);
+			REPORT_PYTHON_GENERATION_ISSUE(Error, TEXT("Struct '%s' is marked as '%s' but the function '%s' could not be found."), *UnresolvedMakePair.Value->Struct->GetPathName(), *PyGenUtil::HasNativeMakeMetaDataKey.ToString(), *UnresolvedMakePair.Key);
 		}
 		for (const TPair<FString, TSharedPtr<FPyWrapperStructMetaData>>& UnresolvedBreakPair : UnresolvedBreakFuncs)
 		{
-			REPORT_PYTHON_GENERATION_ISSUE(Error, TEXT("Struct '%s' is marked as '%s' but the function '%s' could not be found."), *UnresolvedBreakPair.Value->Struct->GetName(), *PyGenUtil::HasNativeBreakMetaDataKey.ToString(), *UnresolvedBreakPair.Key);
+			REPORT_PYTHON_GENERATION_ISSUE(Error, TEXT("Struct '%s' is marked as '%s' but the function '%s' could not be found."), *UnresolvedBreakPair.Value->Struct->GetPathName(), *PyGenUtil::HasNativeBreakMetaDataKey.ToString(), *UnresolvedBreakPair.Key);
 		}
 	});
 }
@@ -846,7 +849,7 @@ void FPyWrapperTypeRegistry::RegisterInlineStructFactory(const TSharedRef<const
 	InlineStructFactories.Add(InFactory->GetStructName(), InFactory);
 }
 
-const IPyWrapperInlineStructFactory* FPyWrapperTypeRegistry::GetInlineStructFactory(const FName StructName) const
+const IPyWrapperInlineStructFactory* FPyWrapperTypeRegistry::GetInlineStructFactory(const FTopLevelAssetPath& StructName) const
 {
 	return InlineStructFactories.FindRef(StructName).Get();
 }
@@ -867,9 +870,14 @@ void FPyWrapperTypeRegistry::GenerateWrappedTypes()
 		TArray<UObject*> ObjectsToProcess;
 		GetObjectsOfClass(UField::StaticClass(), ObjectsToProcess);
 
+		FScopedSlowTask SlowTask(ObjectsToProcess.Num(), LOCTEXT("GeneratingPythonTypes", "Generating Python Types..."));
+		SlowTask.Visibility = ESlowTaskVisibility::Important; // this function can be very slow, users will benefit from our messages
+
 		FNameBuilder ObjectPackageName;
 		for (UObject* ObjectToProcess : ObjectsToProcess)
 		{
+			SlowTask.EnterProgressFrame();
+
 			ObjectToProcess->GetPackage()->GetFName().ToString(ObjectPackageName);
 			if (FPackageName::IsScriptPackage(ObjectPackageName.ToView()))
 			{
@@ -907,8 +915,16 @@ void FPyWrapperTypeRegistry::GenerateWrappedTypesForModule(const FName ModuleNam
 		TArray<UObject*> ObjectsToProcess;
 		GetObjectsWithPackage(ModulePackage, ObjectsToProcess);
 
+		FScopedSlowTask SlowTask(ObjectsToProcess.Num(), FText::Format(LOCTEXT("GeneratingPythonTypesForModule", "Generating Python Types for {0}..."), FText::FromName(ModuleName)));
+		if (ObjectsToProcess.Num() > 1000)
+		{
+			SlowTask.Visibility = ESlowTaskVisibility::Important; // this function can be very slow, users will benefit from our messages
+		}
+
 		for (UObject* ObjectToProcess : ObjectsToProcess)
 		{
+			SlowTask.EnterProgressFrame();
+
 			GenerateWrappedTypeForObject(ObjectToProcess, GeneratedWrappedTypeReferences, DirtyModules);
 		}
 
@@ -943,11 +959,11 @@ void FPyWrapperTypeRegistry::OrphanWrappedTypesForModule(const FName ModuleName)
 {
 	TRACE_CPUPROFILER_EVENT_SCOPE(FPyWrapperTypeRegistry::OrphanWrappedTypesForModule)
 
-	TArray<FName> ModuleTypeNames;
+	TArray<FSoftObjectPath> ModuleTypeNames;
 	GeneratedWrappedTypesForModule.MultiFind(ModuleName, ModuleTypeNames, true);
 	GeneratedWrappedTypesForModule.Remove(ModuleName);
 
-	for (const FName& ModuleTypeName : ModuleTypeNames)
+	for (const FSoftObjectPath& ModuleTypeName : ModuleTypeNames)
 	{
 		TSharedPtr<PyGenUtil::FGeneratedWrappedType> GeneratedWrappedType;
 		if (GeneratedWrappedTypes.RemoveAndCopyValue(ModuleTypeName, GeneratedWrappedType))
@@ -959,6 +975,7 @@ void FPyWrapperTypeRegistry::OrphanWrappedTypesForModule(const FName ModuleName)
 			PythonWrappedClasses.Remove(ModuleTypeName);
 			PythonWrappedStructs.Remove(ModuleTypeName);
 			PythonWrappedEnums.Remove(ModuleTypeName);
+			PythonWrappedDelegates.Remove(ModuleTypeName);
 		}
 	}
 }
@@ -974,25 +991,36 @@ void FPyWrapperTypeRegistry::GenerateWrappedTypesForReferences(const FGeneratedW
 		return;
 	}
 	
+	const int32 TotalNumReferences = InGeneratedWrappedTypeReferences.ClassReferences.Num() + InGeneratedWrappedTypeReferences.StructReferences.Num() + InGeneratedWrappedTypeReferences.EnumReferences.Num() + InGeneratedWrappedTypeReferences.DelegateReferences.Num();
+	FScopedSlowTask SlowTask(TotalNumReferences, LOCTEXT("GeneratingPythonTypes", "Generating Python Types..."));
+	if (TotalNumReferences > 1000)
+	{
+		SlowTask.Visibility = ESlowTaskVisibility::Important; // this function can be very slow, users will benefit from our messages
+	}
+
 	FGeneratedWrappedTypeReferences GeneratedWrappedTypeReferences;
 
 	for (const UClass* Class : InGeneratedWrappedTypeReferences.ClassReferences)
 	{
+		SlowTask.EnterProgressFrame();
 		GenerateWrappedClassType(Class, GeneratedWrappedTypeReferences, OutDirtyModules, ReferenceGenerationFlags);
 	}
 
 	for (const UScriptStruct* Struct : InGeneratedWrappedTypeReferences.StructReferences)
 	{
+		SlowTask.EnterProgressFrame();
 		GenerateWrappedStructType(Struct, GeneratedWrappedTypeReferences, OutDirtyModules, ReferenceGenerationFlags);
 	}
 
 	for (const UEnum* Enum : InGeneratedWrappedTypeReferences.EnumReferences)
 	{
+		SlowTask.EnterProgressFrame();
 		GenerateWrappedEnumType(Enum, GeneratedWrappedTypeReferences, OutDirtyModules, ReferenceGenerationFlags);
 	}
 
 	for (const UFunction* DelegateSignature : InGeneratedWrappedTypeReferences.DelegateReferences)
 	{
+		SlowTask.EnterProgressFrame();
 		checkf(DelegateSignature->HasAnyFunctionFlags(FUNC_Delegate), TEXT("UFunction '%s' was detected as a delegate but doesn't have the 'FUNC_Delegate' flag"), *DelegateSignature->GetPathName());
 		GenerateWrappedDelegateType(DelegateSignature, GeneratedWrappedTypeReferences, OutDirtyModules, ReferenceGenerationFlags);
 	}
@@ -1011,16 +1039,16 @@ void FPyWrapperTypeRegistry::NotifyModulesDirtied(const TSet<FName>& InDirtyModu
 	}
 }
 
-void FPyWrapperTypeRegistry::UpdateGenerateWrappedTypeForRename(const FName InOldTypeRegistryName, const UObject* InObj)
+void FPyWrapperTypeRegistry::UpdateGenerateWrappedTypeForRename(const FSoftObjectPath& InOldTypeRegistryName, const UObject* InObj)
 {
-	const FName TypeRegistryName = PyGenUtil::GetAssetTypeRegistryName(InObj);
+	const FSoftObjectPath TypeRegistryName(InObj);
 
 	TSharedPtr<PyGenUtil::FGeneratedWrappedType> GeneratedWrappedType = GeneratedWrappedTypes.FindAndRemoveChecked(InOldTypeRegistryName);
-	check(!TypeRegistryName.IsNone() && !GeneratedWrappedTypes.Contains(TypeRegistryName));
+	check(TypeRegistryName.IsValid() && !GeneratedWrappedTypes.Contains(TypeRegistryName));
 	GeneratedWrappedTypes.Add(TypeRegistryName, GeneratedWrappedType);
 }
 
-void FPyWrapperTypeRegistry::RemoveGenerateWrappedTypeForDelete(const FName InTypeRegistryName)
+void FPyWrapperTypeRegistry::RemoveGenerateWrappedTypeForDelete(const FSoftObjectPath& InTypeRegistryName)
 {
 	TSharedPtr<PyGenUtil::FGeneratedWrappedType> GeneratedWrappedType = GeneratedWrappedTypes.FindRef(InTypeRegistryName);
 	const FString PythonTypeName = UTF8_TO_TCHAR(GeneratedWrappedType->TypeName.GetData());
@@ -1088,7 +1116,7 @@ bool FPyWrapperTypeRegistry::HasWrappedTypeForObject(const UObject* InObj) const
 	return false;
 }
 
-bool FPyWrapperTypeRegistry::HasWrappedTypeForObjectName(const FName InName) const
+bool FPyWrapperTypeRegistry::HasWrappedTypeForObjectName(const FSoftObjectPath& InName) const
 {
 	return GeneratedWrappedTypes.Contains(InName);
 }
@@ -1127,7 +1155,7 @@ PyTypeObject* FPyWrapperTypeRegistry::GenerateWrappedClassType(const UClass* InC
 	INC_DWORD_STAT(STAT_GenerateWrappedClassCallCount);
 
 	// Already processed? Nothing more to do
-	const FName TypeRegistryName = PyGenUtil::GetTypeRegistryName(InClass);
+	const FSoftObjectPath TypeRegistryName(InClass);
 	if (PyTypeObject* ExistingPyType = PythonWrappedClasses.FindRef(TypeRegistryName))
 	{
 		return ExistingPyType;
@@ -1189,7 +1217,7 @@ PyTypeObject* FPyWrapperTypeRegistry::GenerateWrappedClassType(const UClass* InC
 						const UFunction* GetSetFunc = InClass->FindFunctionByName(*GetSetName);
 						if (!GetSetFunc)
 						{
-							REPORT_PYTHON_GENERATION_ISSUE(Error, TEXT("Property '%s.%s' is marked as '%s' but the function '%s' could not be found."), *InClass->GetName(), *InProp->GetName(), *InKey.ToString(), *GetSetName);
+							REPORT_PYTHON_GENERATION_ISSUE(Error, TEXT("Property '%s.%s' is marked as '%s' but the function '%s' could not be found."), *InClass->GetPathName(), *InProp->GetName(), *InKey.ToString(), *GetSetName);
 						}
 						return GetSetFunc;
 					}
@@ -1212,9 +1240,10 @@ PyTypeObject* FPyWrapperTypeRegistry::GenerateWrappedClassType(const UClass* InC
 				GeneratedWrappedType->FieldTracker.RegisterPythonFieldName(GeneratedPropertyDoc.PythonPropName, InProp);
 			}
 
-			const TArray<FString> DeprecatedPythonPropNames = PyGenUtil::GetDeprecatedPropertyPythonNames(InProp);
-			for (const FString& DeprecatedPythonPropName : DeprecatedPythonPropNames)
+			const TArray<TTuple<FSoftObjectPath, FString>> DeprecatedPythonPropNames = PyGenUtil::GetDeprecatedPropertyPythonNames(InProp);
+			for (const TTuple<FSoftObjectPath, FString>& DeprecatedPythonPropNamePair : DeprecatedPythonPropNames)
 			{
+				const FString& DeprecatedPythonPropName = DeprecatedPythonPropNamePair.Value;
 				FString DeprecationMessage = FString::Printf(TEXT("'%s' was renamed to '%s'."), *DeprecatedPythonPropName, *GeneratedPropertyDoc.PythonPropName);
 				PythonProperties.Add(*DeprecatedPythonPropName, InProp->GetFName());
 				PythonDeprecatedProperties.Add(*DeprecatedPythonPropName, DeprecationMessage);
@@ -1238,7 +1267,7 @@ PyTypeObject* FPyWrapperTypeRegistry::GenerateWrappedClassType(const UClass* InC
 		// Only static functions can be hoisted onto other types
 		if (!InFunc->HasAnyFunctionFlags(FUNC_Static))
 		{
-			REPORT_PYTHON_GENERATION_ISSUE(Error, TEXT("Non-static function '%s.%s' is marked as 'ScriptMethod' but only static functions can be hoisted."), *InFunc->GetOwnerClass()->GetName(), *InFunc->GetName());
+			REPORT_PYTHON_GENERATION_ISSUE(Error, TEXT("Non-static function '%s.%s' is marked as 'ScriptMethod' but only static functions can be hoisted."), *InFunc->GetOwnerClass()->GetPathName(), *InFunc->GetName());
 			return;
 		}
 
@@ -1264,14 +1293,14 @@ PyTypeObject* FPyWrapperTypeRegistry::GenerateWrappedClassType(const UClass* InC
 				break;
 			}
 
-			REPORT_PYTHON_GENERATION_ISSUE(Error, TEXT("Function '%s.%s' is marked as 'ScriptMethod' but doesn't contain a valid struct or object as its first argument.%s"), *InFunc->GetOwnerClass()->GetName(), *InFunc->GetName(), PropertyPossiblyMissingMacro ? TEXT(" UPARAM(ref) may be missing on the first argument.") : TEXT(""));
+			REPORT_PYTHON_GENERATION_ISSUE(Error, TEXT("Function '%s.%s' is marked as 'ScriptMethod' but doesn't contain a valid struct or object as its first argument.%s"), *InFunc->GetOwnerClass()->GetPathName(), *InFunc->GetName(), PropertyPossiblyMissingMacro ? TEXT(" UPARAM(ref) may be missing on the first argument.") : TEXT(""));
 			return;
 		}
 		if (const FObjectPropertyBase* SelfPropObj = CastField<FObjectPropertyBase>(SelfParam.ParamProp))
 		{
 			if (SelfPropObj->PropertyClass->IsChildOf(InFunc->GetOwnerClass()))
 			{
-				REPORT_PYTHON_GENERATION_ISSUE(Error, TEXT("Function '%s.%s' is marked as 'ScriptMethod' but the object argument type (%s) is a child of the the class type of the static function. This is not allowed."), *InFunc->GetOwnerClass()->GetName(), *InFunc->GetName(), *SelfPropObj->PropertyClass->GetName());
+				REPORT_PYTHON_GENERATION_ISSUE(Error, TEXT("Function '%s.%s' is marked as 'ScriptMethod' but the object argument type (%s) is a child of the the class type of the static function. This is not allowed."), *InFunc->GetOwnerClass()->GetPathName(), *InFunc->GetName(), *SelfPropObj->PropertyClass->GetPathName());
 				return;
 			}
 		}
@@ -1310,27 +1339,27 @@ PyTypeObject* FPyWrapperTypeRegistry::GenerateWrappedClassType(const UClass* InC
 		{
 			if (GeneratedWrappedDynamicMethod.SelfReturn.ParamProp)
 			{
-				REPORT_PYTHON_GENERATION_ISSUE(Error, TEXT("Function '%s.%s' is marked as 'ScriptMethodSelfReturn' but the 'self' argument is also marked as UPARAM(ref). This is not allowed."), *InFunc->GetOwnerClass()->GetName(), *InFunc->GetName());
+				REPORT_PYTHON_GENERATION_ISSUE(Error, TEXT("Function '%s.%s' is marked as 'ScriptMethodSelfReturn' but the 'self' argument is also marked as UPARAM(ref). This is not allowed."), *InFunc->GetOwnerClass()->GetPathName(), *InFunc->GetName());
 				return;
 			}
 			else if (GeneratedWrappedDynamicMethod.MethodFunc.OutputParams.Num() == 0 || !GeneratedWrappedDynamicMethod.MethodFunc.OutputParams[0].ParamProp->HasAnyPropertyFlags(CPF_ReturnParm))
 			{
-				REPORT_PYTHON_GENERATION_ISSUE(Error, TEXT("Function '%s.%s' is marked as 'ScriptMethodSelfReturn' but has no return value."), *InFunc->GetOwnerClass()->GetName(), *InFunc->GetName());
+				REPORT_PYTHON_GENERATION_ISSUE(Error, TEXT("Function '%s.%s' is marked as 'ScriptMethodSelfReturn' but has no return value."), *InFunc->GetOwnerClass()->GetPathName(), *InFunc->GetName());
 				return;
 			}
 			else if (!SelfParam.ParamProp->IsA<FStructProperty>())
 			{
-				REPORT_PYTHON_GENERATION_ISSUE(Error, TEXT("Function '%s.%s' is marked as 'ScriptMethodSelfReturn' but the 'self' argument is not a struct."), *InFunc->GetOwnerClass()->GetName(), *InFunc->GetName());
+				REPORT_PYTHON_GENERATION_ISSUE(Error, TEXT("Function '%s.%s' is marked as 'ScriptMethodSelfReturn' but the 'self' argument is not a struct."), *InFunc->GetOwnerClass()->GetPathName(), *InFunc->GetName());
 				return;
 			}
 			else if (!GeneratedWrappedDynamicMethod.MethodFunc.OutputParams[0].ParamProp->IsA<FStructProperty>())
 			{
-				REPORT_PYTHON_GENERATION_ISSUE(Error, TEXT("Function '%s.%s' is marked as 'ScriptMethodSelfReturn' but the return value is not a struct."), *InFunc->GetOwnerClass()->GetName(), *InFunc->GetName());
+				REPORT_PYTHON_GENERATION_ISSUE(Error, TEXT("Function '%s.%s' is marked as 'ScriptMethodSelfReturn' but the return value is not a struct."), *InFunc->GetOwnerClass()->GetPathName(), *InFunc->GetName());
 				return;
 			}
 			else if (CastFieldChecked<const FStructProperty>(GeneratedWrappedDynamicMethod.MethodFunc.OutputParams[0].ParamProp)->Struct != CastFieldChecked<const FStructProperty>(SelfParam.ParamProp)->Struct)
 			{
-				REPORT_PYTHON_GENERATION_ISSUE(Error, TEXT("Function '%s.%s' is marked as 'ScriptMethodSelfReturn' but the return value is not the same type as the 'self' argument."), *InFunc->GetOwnerClass()->GetName(), *InFunc->GetName());
+				REPORT_PYTHON_GENERATION_ISSUE(Error, TEXT("Function '%s.%s' is marked as 'ScriptMethodSelfReturn' but the return value is not the same type as the 'self' argument."), *InFunc->GetOwnerClass()->GetPathName(), *InFunc->GetName());
 				return;
 			}
 			else
@@ -1345,7 +1374,7 @@ PyTypeObject* FPyWrapperTypeRegistry::GenerateWrappedClassType(const UClass* InC
 		{
 			if (!SelfParam.ParamProp->IsA<FStructProperty>())
 			{
-				REPORT_PYTHON_GENERATION_ISSUE(Error, TEXT("Function '%s.%s' is marked as 'ScriptMethodMutable' but the 'self' argument is not a struct."), *InFunc->GetOwnerClass()->GetName(), *InFunc->GetName());
+				REPORT_PYTHON_GENERATION_ISSUE(Error, TEXT("Function '%s.%s' is marked as 'ScriptMethodMutable' but the 'self' argument is not a struct."), *InFunc->GetOwnerClass()->GetPathName(), *InFunc->GetName());
 				return;
 			}
 			GeneratedWrappedDynamicMethod.SelfReturn = SelfParam;
@@ -1377,9 +1406,10 @@ PyTypeObject* FPyWrapperTypeRegistry::GenerateWrappedClassType(const UClass* InC
 		}
 
 		// Add any deprecated variants too
-		const TArray<FString> DeprecatedPythonStructMethodNames = PyGenUtil::GetDeprecatedScriptMethodPythonNames(InFunc);
-		for (const FString& DeprecatedPythonStructMethodName : DeprecatedPythonStructMethodNames)
+		const TArray<TTuple<FSoftObjectPath, FString>> DeprecatedPythonStructMethodNames = PyGenUtil::GetDeprecatedScriptMethodPythonNames(InFunc);
+		for (const TTuple<FSoftObjectPath, FString>& DeprecatedPythonStructMethodNamePair : DeprecatedPythonStructMethodNames)
 		{
+			const FString& DeprecatedPythonStructMethodName = DeprecatedPythonStructMethodNamePair.Value;
 			FString DeprecationMessage = FString::Printf(TEXT("'%s' was renamed to '%s'."), *DeprecatedPythonStructMethodName, *PythonStructMethodName);
 
 			PyGenUtil::FGeneratedWrappedDynamicMethod& DeprecatedGeneratedWrappedMethod = DynamicMethodDefs.AddDefaulted_GetRef().Get();
@@ -1400,7 +1430,7 @@ PyTypeObject* FPyWrapperTypeRegistry::GenerateWrappedClassType(const UClass* InC
 			}
 
 			// Find the wrapped type for the class as that's what we'll actually add the dynamic method to
-			TSharedPtr<PyGenUtil::FGeneratedWrappedClassType> HostedClassGeneratedWrappedType = StaticCastSharedPtr<PyGenUtil::FGeneratedWrappedClassType>(GeneratedWrappedTypes.FindRef(PyGenUtil::GetTypeRegistryName(HostedClass)));
+			TSharedPtr<PyGenUtil::FGeneratedWrappedClassType> HostedClassGeneratedWrappedType = StaticCastSharedPtr<PyGenUtil::FGeneratedWrappedClassType>(GeneratedWrappedTypes.FindRef(HostedClass));
 			check(HostedClassGeneratedWrappedType.IsValid());
 
 			// Add the dynamic methods to the class
@@ -1420,7 +1450,7 @@ PyTypeObject* FPyWrapperTypeRegistry::GenerateWrappedClassType(const UClass* InC
 			}
 
 			// Find the wrapped type for the struct as that's what we'll actually add the dynamic method to
-			TSharedPtr<PyGenUtil::FGeneratedWrappedStructType> HostedStructGeneratedWrappedType = StaticCastSharedPtr<PyGenUtil::FGeneratedWrappedStructType>(GeneratedWrappedTypes.FindRef(PyGenUtil::GetTypeRegistryName(HostedStruct)));
+			TSharedPtr<PyGenUtil::FGeneratedWrappedStructType> HostedStructGeneratedWrappedType = StaticCastSharedPtr<PyGenUtil::FGeneratedWrappedStructType>(GeneratedWrappedTypes.FindRef(HostedStruct));
 			check(HostedStructGeneratedWrappedType.IsValid());
 
 			// Add the dynamic methods to the struct
@@ -1441,7 +1471,7 @@ PyTypeObject* FPyWrapperTypeRegistry::GenerateWrappedClassType(const UClass* InC
 		// Only static functions can be hoisted onto other types
 		if (!InFunc->HasAnyFunctionFlags(FUNC_Static))
 		{
-			REPORT_PYTHON_GENERATION_ISSUE(Error, TEXT("Non-static function '%s.%s' is marked as 'ScriptOperator' but only static functions can be hoisted."), *InFunc->GetOwnerClass()->GetName(), *InFunc->GetName());
+			REPORT_PYTHON_GENERATION_ISSUE(Error, TEXT("Non-static function '%s.%s' is marked as 'ScriptOperator' but only static functions can be hoisted."), *InFunc->GetOwnerClass()->GetPathName(), *InFunc->GetName());
 			return;
 		}
 
@@ -1458,7 +1488,7 @@ PyTypeObject* FPyWrapperTypeRegistry::GenerateWrappedClassType(const UClass* InC
 			PyGenUtil::FGeneratedWrappedOperatorSignature OpSignature;
 			if (!PyGenUtil::FGeneratedWrappedOperatorSignature::StringToSignature(*ScriptOperator, OpSignature))
 			{
-				REPORT_PYTHON_GENERATION_ISSUE(Error, TEXT("Function '%s.%s' is marked as 'ScriptOperator' but uses an unknown operator type '%s'."), *InFunc->GetOwnerClass()->GetName(), *InFunc->GetName(), *ScriptOperator);
+				REPORT_PYTHON_GENERATION_ISSUE(Error, TEXT("Function '%s.%s' is marked as 'ScriptOperator' but uses an unknown operator type '%s'."), *InFunc->GetOwnerClass()->GetPathName(), *InFunc->GetName(), *ScriptOperator);
 				continue;
 			}
 
@@ -1467,7 +1497,7 @@ PyTypeObject* FPyWrapperTypeRegistry::GenerateWrappedClassType(const UClass* InC
 				FString SignatureError;
 				if (!OpFunc.SetFunction(InTypeMethod.MethodFunc, OpSignature, &SignatureError))
 				{
-					REPORT_PYTHON_GENERATION_ISSUE(Error, TEXT("Function '%s.%s' is marked as 'ScriptOperator' but has an invalid signature for the '%s' operator: %s."), *InFunc->GetOwnerClass()->GetName(), *InFunc->GetName(), *ScriptOperator, *SignatureError);
+					REPORT_PYTHON_GENERATION_ISSUE(Error, TEXT("Function '%s.%s' is marked as 'ScriptOperator' but has an invalid signature for the '%s' operator: %s."), *InFunc->GetOwnerClass()->GetPathName(), *InFunc->GetName(), *ScriptOperator, *SignatureError);
 					continue;
 				}
 			}
@@ -1477,7 +1507,7 @@ PyTypeObject* FPyWrapperTypeRegistry::GenerateWrappedClassType(const UClass* InC
 			if (GenerateWrappedStructType(HostedStruct, OutGeneratedWrappedTypeReferences, OutDirtyModules, EPyTypeGenerationFlags::ForceShouldExport))
 			{
 				// Find the wrapped type for the struct as that's what we'll actually add the operator to (via its meta-data)
-				TSharedPtr<PyGenUtil::FGeneratedWrappedStructType> HostedStructGeneratedWrappedType = StaticCastSharedPtr<PyGenUtil::FGeneratedWrappedStructType>(GeneratedWrappedTypes.FindRef(PyGenUtil::GetTypeRegistryName(HostedStruct)));
+				TSharedPtr<PyGenUtil::FGeneratedWrappedStructType> HostedStructGeneratedWrappedType = StaticCastSharedPtr<PyGenUtil::FGeneratedWrappedStructType>(GeneratedWrappedTypes.FindRef(HostedStruct));
 				check(HostedStructGeneratedWrappedType.IsValid());
 				StaticCastSharedPtr<FPyWrapperStructMetaData>(HostedStructGeneratedWrappedType->MetaData)->OpStacks[(int32)OpSignature.OpType].Funcs.Add(MoveTemp(OpFunc));
 			}
@@ -1489,7 +1519,7 @@ PyTypeObject* FPyWrapperTypeRegistry::GenerateWrappedClassType(const UClass* InC
 		// Only static functions can be constants
 		if (!InFunc->HasAnyFunctionFlags(FUNC_Static))
 		{
-			REPORT_PYTHON_GENERATION_ISSUE(Error, TEXT("Non-static function '%s.%s' is marked as 'ScriptConstant' but only static functions can be hoisted."), *InFunc->GetOwnerClass()->GetName(), *InFunc->GetName());
+			REPORT_PYTHON_GENERATION_ISSUE(Error, TEXT("Non-static function '%s.%s' is marked as 'ScriptConstant' but only static functions can be hoisted."), *InFunc->GetOwnerClass()->GetPathName(), *InFunc->GetName());
 			return;
 		}
 
@@ -1505,7 +1535,7 @@ PyTypeObject* FPyWrapperTypeRegistry::GenerateWrappedClassType(const UClass* InC
 			}
 			if (!HostType)
 			{
-				REPORT_PYTHON_GENERATION_ISSUE(Error, TEXT("Function '%s.%s' is marked as 'ScriptConstantHost' but the host '%s' could not be found."), *InFunc->GetOwnerClass()->GetName(), *InFunc->GetName(), *ConstantOwnerName);
+				REPORT_PYTHON_GENERATION_ISSUE(Error, TEXT("Function '%s.%s' is marked as 'ScriptConstantHost' but the host '%s' could not be found."), *InFunc->GetOwnerClass()->GetPathName(), *InFunc->GetName(), *ConstantOwnerName);
 				return;
 			}
 		}
@@ -1513,7 +1543,7 @@ PyTypeObject* FPyWrapperTypeRegistry::GenerateWrappedClassType(const UClass* InC
 		{
 			if (HostClass->IsChildOf(InFunc->GetOwnerClass()))
 			{
-				REPORT_PYTHON_GENERATION_ISSUE(Error, TEXT("Function '%s.%s' is marked as 'ScriptConstantHost' but the host type (%s) is a child of the class type of the static function. This is not allowed."), *InFunc->GetOwnerClass()->GetName(), *InFunc->GetName(), *HostClass->GetName());
+				REPORT_PYTHON_GENERATION_ISSUE(Error, TEXT("Function '%s.%s' is marked as 'ScriptConstantHost' but the host type (%s) is a child of the class type of the static function. This is not allowed."), *InFunc->GetOwnerClass()->GetPathName(), *InFunc->GetName(), *HostClass->GetPathName());
 				return;
 			}
 		}
@@ -1523,7 +1553,7 @@ PyTypeObject* FPyWrapperTypeRegistry::GenerateWrappedClassType(const UClass* InC
 		ConstantFunc.SetFunction(InFunc);
 		if (ConstantFunc.InputParams.Num() != 0 || ConstantFunc.OutputParams.Num() != 1)
 		{
-			REPORT_PYTHON_GENERATION_ISSUE(Error, TEXT("Function '%s.%s' is marked as 'ScriptConstant' but has an invalid signature (it must return a value and take no arguments)."), *InFunc->GetOwnerClass()->GetName(), *InFunc->GetName());
+			REPORT_PYTHON_GENERATION_ISSUE(Error, TEXT("Function '%s.%s' is marked as 'ScriptConstant' but has an invalid signature (it must return a value and take no arguments)."), *InFunc->GetOwnerClass()->GetPathName(), *InFunc->GetName());
 			return;
 		}
 
@@ -1537,9 +1567,10 @@ PyTypeObject* FPyWrapperTypeRegistry::GenerateWrappedClassType(const UClass* InC
 		GeneratedWrappedConstant.ConstantFunc = ConstantFunc;
 
 		// Build any deprecated variants too
-		const TArray<FString> DeprecatedPythonConstantNames = PyGenUtil::GetDeprecatedScriptConstantPythonNames(InFunc);
-		for (const FString& DeprecatedPythonConstantName : DeprecatedPythonConstantNames)
+		const TArray<TTuple<FSoftObjectPath, FString>> DeprecatedPythonConstantNames = PyGenUtil::GetDeprecatedScriptConstantPythonNames(InFunc);
+		for (const TTuple<FSoftObjectPath, FString>& DeprecatedPythonConstantNamePair : DeprecatedPythonConstantNames)
 		{
+			const FString& DeprecatedPythonConstantName = DeprecatedPythonConstantNamePair.Value;
 			FString DeprecationMessage = FString::Printf(TEXT("'%s' was renamed to '%s'."), *DeprecatedPythonConstantName, *PythonConstantName);
 
 			PyGenUtil::FGeneratedWrappedConstant& DeprecatedGeneratedWrappedConstant = ConstantDefs.AddDefaulted_GetRef().Get();
@@ -1562,7 +1593,7 @@ PyTypeObject* FPyWrapperTypeRegistry::GenerateWrappedClassType(const UClass* InC
 				}
 
 				// Find the wrapped type for the class as that's what we'll actually add the constant to
-				TSharedPtr<PyGenUtil::FGeneratedWrappedClassType> HostedClassGeneratedWrappedType = StaticCastSharedPtr<PyGenUtil::FGeneratedWrappedClassType>(GeneratedWrappedTypes.FindRef(PyGenUtil::GetTypeRegistryName(HostClass)));
+				TSharedPtr<PyGenUtil::FGeneratedWrappedClassType> HostedClassGeneratedWrappedType = StaticCastSharedPtr<PyGenUtil::FGeneratedWrappedClassType>(GeneratedWrappedTypes.FindRef(HostClass));
 				check(HostedClassGeneratedWrappedType.IsValid());
 
 				// Add the dynamic constants to the struct
@@ -1583,7 +1614,7 @@ PyTypeObject* FPyWrapperTypeRegistry::GenerateWrappedClassType(const UClass* InC
 				}
 
 				// Find the wrapped type for the struct as that's what we'll actually add the constant to
-				TSharedPtr<PyGenUtil::FGeneratedWrappedStructType> HostedStructGeneratedWrappedType = StaticCastSharedPtr<PyGenUtil::FGeneratedWrappedStructType>(GeneratedWrappedTypes.FindRef(PyGenUtil::GetTypeRegistryName(HostStruct)));
+				TSharedPtr<PyGenUtil::FGeneratedWrappedStructType> HostedStructGeneratedWrappedType = StaticCastSharedPtr<PyGenUtil::FGeneratedWrappedStructType>(GeneratedWrappedTypes.FindRef(HostStruct));
 				check(HostedStructGeneratedWrappedType.IsValid());
 
 				// Add the dynamic constants to the struct
@@ -1665,9 +1696,10 @@ PyTypeObject* FPyWrapperTypeRegistry::GenerateWrappedClassType(const UClass* InC
 
 		const PyGenUtil::FGeneratedWrappedMethod GeneratedWrappedMethodCopy = GeneratedWrappedMethod;
 
-		const TArray<FString> DeprecatedPythonFuncNames = PyGenUtil::GetDeprecatedFunctionPythonNames(InFunc);
-		for (const FString& DeprecatedPythonFuncName : DeprecatedPythonFuncNames)
+		const TArray<TTuple<FSoftObjectPath, FString>> DeprecatedPythonFuncNames = PyGenUtil::GetDeprecatedFunctionPythonNames(InFunc);
+		for (const TTuple<FSoftObjectPath, FString>& DeprecatedPythonFuncNamePair : DeprecatedPythonFuncNames)
 		{
+			const FString& DeprecatedPythonFuncName = DeprecatedPythonFuncNamePair.Value;
 			FString DeprecationMessage = FString::Printf(TEXT("'%s' was renamed to '%s'."), *DeprecatedPythonFuncName, *PythonFunctionName);
 			PythonMethods.Add(*DeprecatedPythonFuncName, InFunc->GetFName());
 			PythonDeprecatedMethods.Add(*DeprecatedPythonFuncName, DeprecationMessage);
@@ -1727,7 +1759,7 @@ PyTypeObject* FPyWrapperTypeRegistry::GenerateWrappedClassType(const UClass* InC
 	FString TypeDocString = PyGenUtil::PythonizeTooltip(PyGenUtil::ParseTooltip(PyGenUtil::GetFieldTooltip(InClass)));
 	if (const UClass* SuperClass = InClass->GetSuperClass())
 	{
-		TSharedPtr<PyGenUtil::FGeneratedWrappedClassType> SuperGeneratedWrappedType = StaticCastSharedPtr<PyGenUtil::FGeneratedWrappedClassType>(GeneratedWrappedTypes.FindRef(PyGenUtil::GetTypeRegistryName(SuperClass)));
+		TSharedPtr<PyGenUtil::FGeneratedWrappedClassType> SuperGeneratedWrappedType = StaticCastSharedPtr<PyGenUtil::FGeneratedWrappedClassType>(GeneratedWrappedTypes.FindRef(SuperClass));
 		if (SuperGeneratedWrappedType.IsValid())
 		{
 			GeneratedWrappedType->PropertyDocs.Append(SuperGeneratedWrappedType->PropertyDocs);
@@ -1779,20 +1811,21 @@ PyTypeObject* FPyWrapperTypeRegistry::GenerateWrappedClassType(const UClass* InC
 		RegisterWrappedClassType(TypeRegistryName, &GeneratedWrappedType->PyType, !bIsBlueprintGeneratedType);
 
 		// Also generate and register any deprecated aliases for this type
-		const TArray<FString> DeprecatedPythonClassNames = PyGenUtil::GetDeprecatedClassPythonNames(InClass);
-		for (const FString& DeprecatedPythonClassName : DeprecatedPythonClassNames)
+		const TArray<TTuple<FSoftObjectPath, FString>> DeprecatedPythonClassNames = PyGenUtil::GetDeprecatedClassPythonNames(InClass);
+		for (const TTuple<FSoftObjectPath, FString>& DeprecatedPythonClassNamePair : DeprecatedPythonClassNames)
 		{
-			const FName DeprecatedClassName = *DeprecatedPythonClassName;
+			const FString& DeprecatedPythonClassName = DeprecatedPythonClassNamePair.Value;
 			FString DeprecationMessage = FString::Printf(TEXT("'%s' was renamed to '%s'."), *DeprecatedPythonClassName, *PythonClassName);
 			
-			if (GeneratedWrappedTypes.Contains(DeprecatedClassName))
+			const FSoftObjectPath& DeprecatedTypeRegistryName = DeprecatedPythonClassNamePair.Key;
+			if (GeneratedWrappedTypes.Contains(DeprecatedTypeRegistryName))
 			{
-				REPORT_PYTHON_GENERATION_ISSUE(Warning, TEXT("Deprecated class name '%s' conflicted with an existing type!"), *DeprecatedPythonClassName);
+				REPORT_PYTHON_GENERATION_ISSUE(Warning, TEXT("Deprecated class name '%s' conflicted with an existing type!"), *DeprecatedTypeRegistryName.ToString());
 				continue;
 			}
 
 			TSharedRef<PyGenUtil::FGeneratedWrappedClassType> DeprecatedGeneratedWrappedType = MakeShared<PyGenUtil::FGeneratedWrappedClassType>();
-			GeneratedWrappedTypes.Add(DeprecatedClassName, DeprecatedGeneratedWrappedType);
+			GeneratedWrappedTypes.Add(DeprecatedTypeRegistryName, DeprecatedGeneratedWrappedType);
 
 			DeprecatedGeneratedWrappedType->TypeName = PyGenUtil::TCHARToUTF8Buffer(*DeprecatedPythonClassName);
 			DeprecatedGeneratedWrappedType->TypeDoc = PyGenUtil::TCHARToUTF8Buffer(*FString::Printf(TEXT("deprecated: %s"), *DeprecationMessage));
@@ -1808,7 +1841,7 @@ PyTypeObject* FPyWrapperTypeRegistry::GenerateWrappedClassType(const UClass* InC
 			{
 				if (!UnrealModuleName.IsNone())
 				{
-					GeneratedWrappedTypesForModule.Add(UnrealModuleName, DeprecatedClassName);
+					GeneratedWrappedTypesForModule.Add(UnrealModuleName, DeprecatedTypeRegistryName);
 					// Execute Python code within this block
 					{
 						FPyScopedGIL GIL;
@@ -1818,22 +1851,22 @@ PyTypeObject* FPyWrapperTypeRegistry::GenerateWrappedClassType(const UClass* InC
 						PyModule_AddObject(PyModule, DeprecatedGeneratedWrappedType->PyType.tp_name, (PyObject*)&DeprecatedGeneratedWrappedType->PyType);
 					}
 				}
-				RegisterWrappedClassType(DeprecatedClassName, &DeprecatedGeneratedWrappedType->PyType, !bIsBlueprintGeneratedType);
+				RegisterWrappedClassType(DeprecatedTypeRegistryName, &DeprecatedGeneratedWrappedType->PyType, !bIsBlueprintGeneratedType);
 			}
 			else
 			{
-				REPORT_PYTHON_GENERATION_ISSUE(Fatal, TEXT("Failed to generate Python glue code for deprecated class '%s'!"), *DeprecatedPythonClassName);
+				REPORT_PYTHON_GENERATION_ISSUE(Fatal, TEXT("Failed to generate Python glue code for deprecated class '%s'!"), *DeprecatedTypeRegistryName.ToString());
 			}
 		}
 
 		return &GeneratedWrappedType->PyType;
 	}
 	
-	REPORT_PYTHON_GENERATION_ISSUE(Fatal, TEXT("Failed to generate Python glue code for class '%s'!"), *InClass->GetName());
+	REPORT_PYTHON_GENERATION_ISSUE(Fatal, TEXT("Failed to generate Python glue code for class '%s'!"), *InClass->GetPathName());
 	return nullptr;
 }
 
-void FPyWrapperTypeRegistry::RegisterWrappedClassType(const FName ClassName, PyTypeObject* PyType, const bool InDetectNameConflicts)
+void FPyWrapperTypeRegistry::RegisterWrappedClassType(const FSoftObjectPath& ClassName, PyTypeObject* PyType, const bool InDetectNameConflicts)
 {
 	if (InDetectNameConflicts)
 {
@@ -1842,7 +1875,7 @@ void FPyWrapperTypeRegistry::RegisterWrappedClassType(const FName ClassName, PyT
 	PythonWrappedClasses.Add(ClassName, PyType);
 }
 
-void FPyWrapperTypeRegistry::UnregisterWrappedClassType(const FName ClassName, PyTypeObject* PyType, const bool InUnregisterName)
+void FPyWrapperTypeRegistry::UnregisterWrappedClassType(const FSoftObjectPath& ClassName, PyTypeObject* PyType, const bool InUnregisterName)
 {
 	if (InUnregisterName)
 	{
@@ -1853,8 +1886,7 @@ void FPyWrapperTypeRegistry::UnregisterWrappedClassType(const FName ClassName, P
 
 bool FPyWrapperTypeRegistry::HasWrappedClassType(const UClass* InClass) const
 {
-	const FName TypeRegistryName = PyGenUtil::GetTypeRegistryName(InClass);
-	return PythonWrappedClasses.Contains(TypeRegistryName);
+	return PythonWrappedClasses.Contains(InClass);
 }
 
 PyTypeObject* FPyWrapperTypeRegistry::GetWrappedClassType(const UClass* InClass) const
@@ -1863,8 +1895,7 @@ PyTypeObject* FPyWrapperTypeRegistry::GetWrappedClassType(const UClass* InClass)
 
 	for (const UClass* Class = InClass; Class; Class = Class->GetSuperClass())
 	{
-		const FName TypeRegistryName = PyGenUtil::GetTypeRegistryName(Class);
-		if (PyTypeObject* ClassPyType = PythonWrappedClasses.FindRef(TypeRegistryName))
+		if (PyTypeObject* ClassPyType = PythonWrappedClasses.FindRef(Class))
 		{
 			PyType = ClassPyType;
 			break;
@@ -1897,7 +1928,7 @@ PyTypeObject* FPyWrapperTypeRegistry::GenerateWrappedStructType(const UScriptStr
 	};
 
 	// Already processed? Nothing more to do
-	const FName TypeRegistryName = PyGenUtil::GetTypeRegistryName(InStruct);
+	const FSoftObjectPath TypeRegistryName(InStruct);
 	if (PyTypeObject* ExistingPyType = PythonWrappedStructs.FindRef(TypeRegistryName))
 	{
 		return ExistingPyType;
@@ -1963,9 +1994,10 @@ PyTypeObject* FPyWrapperTypeRegistry::GenerateWrappedStructType(const UScriptStr
 				GeneratedWrappedType->FieldTracker.RegisterPythonFieldName(GeneratedPropertyDoc.PythonPropName, InProp);
 			}
 
-			const TArray<FString> DeprecatedPythonPropNames = PyGenUtil::GetDeprecatedPropertyPythonNames(InProp);
-			for (const FString& DeprecatedPythonPropName : DeprecatedPythonPropNames)
+			const TArray<TTuple<FSoftObjectPath, FString>> DeprecatedPythonPropNames = PyGenUtil::GetDeprecatedPropertyPythonNames(InProp);
+			for (const TTuple<FSoftObjectPath, FString>& DeprecatedPythonPropNamePair : DeprecatedPythonPropNames)
 			{
+				const FString& DeprecatedPythonPropName = DeprecatedPythonPropNamePair.Value;
 				FString DeprecationMessage = FString::Printf(TEXT("'%s' was renamed to '%s'."), *DeprecatedPythonPropName, *GeneratedPropertyDoc.PythonPropName);
 				PythonProperties.Add(*DeprecatedPythonPropName, InProp->GetFName());
 				PythonDeprecatedProperties.Add(*DeprecatedPythonPropName, DeprecationMessage);
@@ -1996,7 +2028,7 @@ PyTypeObject* FPyWrapperTypeRegistry::GenerateWrappedStructType(const UScriptStr
 	FString TypeDocString = PyGenUtil::PythonizeTooltip(PyGenUtil::ParseTooltip(PyGenUtil::GetFieldTooltip(InStruct)));
 	if (const UScriptStruct* SuperStruct = Cast<UScriptStruct>(InStruct->GetSuperStruct()))
 	{
-		TSharedPtr<PyGenUtil::FGeneratedWrappedStructType> SuperGeneratedWrappedType = StaticCastSharedPtr<PyGenUtil::FGeneratedWrappedStructType>(GeneratedWrappedTypes.FindRef(PyGenUtil::GetTypeRegistryName(SuperStruct)));
+		TSharedPtr<PyGenUtil::FGeneratedWrappedStructType> SuperGeneratedWrappedType = StaticCastSharedPtr<PyGenUtil::FGeneratedWrappedStructType>(GeneratedWrappedTypes.FindRef(SuperStruct));
 		if (SuperGeneratedWrappedType.IsValid())
 		{
 			GeneratedWrappedType->PropertyDocs.Append(SuperGeneratedWrappedType->PropertyDocs);
@@ -2008,7 +2040,7 @@ PyTypeObject* FPyWrapperTypeRegistry::GenerateWrappedStructType(const UScriptStr
 	GeneratedWrappedType->TypeDoc = PyGenUtil::TCHARToUTF8Buffer(*TypeDocString);
 
 	int32 WrappedStructSizeBytes = sizeof(FPyWrapperStruct);
-	if (const IPyWrapperInlineStructFactory* InlineStructFactory = GetInlineStructFactory(TypeRegistryName))
+	if (const IPyWrapperInlineStructFactory* InlineStructFactory = GetInlineStructFactory(InStruct->GetStructPathName()))
 	{
 		WrappedStructSizeBytes = InlineStructFactory->GetPythonObjectSizeBytes();
 	}
@@ -2080,20 +2112,21 @@ PyTypeObject* FPyWrapperTypeRegistry::GenerateWrappedStructType(const UScriptStr
 		RegisterWrappedStructType(TypeRegistryName, &GeneratedWrappedType->PyType, !bIsBlueprintGeneratedType);
 
 		// Also generate and register any deprecated aliases for this type
-		const TArray<FString> DeprecatedPythonStructNames = PyGenUtil::GetDeprecatedStructPythonNames(InStruct);
-		for (const FString& DeprecatedPythonStructName : DeprecatedPythonStructNames)
+		const TArray<TTuple<FSoftObjectPath, FString>> DeprecatedPythonStructNames = PyGenUtil::GetDeprecatedStructPythonNames(InStruct);
+		for (const TTuple<FSoftObjectPath, FString>& DeprecatedPythonStructNamePair : DeprecatedPythonStructNames)
 		{
-			const FName DeprecatedStructName = *DeprecatedPythonStructName;
+			const FString& DeprecatedPythonStructName = DeprecatedPythonStructNamePair.Value;
 			FString DeprecationMessage = FString::Printf(TEXT("'%s' was renamed to '%s'."), *DeprecatedPythonStructName, *PythonStructName);
 
-			if (GeneratedWrappedTypes.Contains(DeprecatedStructName))
+			const FSoftObjectPath& DeprecatedTypeRegistryName = DeprecatedPythonStructNamePair.Key;
+			if (GeneratedWrappedTypes.Contains(DeprecatedTypeRegistryName))
 			{
-				REPORT_PYTHON_GENERATION_ISSUE(Warning, TEXT("Deprecated struct name '%s' conflicted with an existing type!"), *DeprecatedPythonStructName);
+				REPORT_PYTHON_GENERATION_ISSUE(Warning, TEXT("Deprecated struct name '%s' conflicted with an existing type!"), *DeprecatedTypeRegistryName.ToString());
 				continue;
 			}
 
 			TSharedRef<PyGenUtil::FGeneratedWrappedStructType> DeprecatedGeneratedWrappedType = MakeShared<PyGenUtil::FGeneratedWrappedStructType>();
-			GeneratedWrappedTypes.Add(DeprecatedStructName, DeprecatedGeneratedWrappedType);
+			GeneratedWrappedTypes.Add(DeprecatedTypeRegistryName, DeprecatedGeneratedWrappedType);
 
 			DeprecatedGeneratedWrappedType->TypeName = PyGenUtil::TCHARToUTF8Buffer(*DeprecatedPythonStructName);
 			DeprecatedGeneratedWrappedType->TypeDoc = PyGenUtil::TCHARToUTF8Buffer(*FString::Printf(TEXT("deprecated: %s"), *DeprecationMessage));
@@ -2109,7 +2142,7 @@ PyTypeObject* FPyWrapperTypeRegistry::GenerateWrappedStructType(const UScriptStr
 			{
 				if (!UnrealModuleName.IsNone())
 				{
-					GeneratedWrappedTypesForModule.Add(UnrealModuleName, DeprecatedStructName);
+					GeneratedWrappedTypesForModule.Add(UnrealModuleName, DeprecatedTypeRegistryName);
 					// Execute Python code within this block
 					{
 						FPyScopedGIL GIL;
@@ -2119,22 +2152,22 @@ PyTypeObject* FPyWrapperTypeRegistry::GenerateWrappedStructType(const UScriptStr
 						PyModule_AddObject(PyModule, DeprecatedGeneratedWrappedType->PyType.tp_name, (PyObject*)&DeprecatedGeneratedWrappedType->PyType);
 					}
 				}
-				RegisterWrappedStructType(DeprecatedStructName, &DeprecatedGeneratedWrappedType->PyType, !bIsBlueprintGeneratedType);
+				RegisterWrappedStructType(DeprecatedTypeRegistryName, &DeprecatedGeneratedWrappedType->PyType, !bIsBlueprintGeneratedType);
 			}
 			else
 			{
-				REPORT_PYTHON_GENERATION_ISSUE(Fatal, TEXT("Failed to generate Python glue code for deprecated struct '%s'!"), *DeprecatedPythonStructName);
+				REPORT_PYTHON_GENERATION_ISSUE(Fatal, TEXT("Failed to generate Python glue code for deprecated struct '%s'!"), *DeprecatedTypeRegistryName.ToString());
 			}
 		}
 
 		return &GeneratedWrappedType->PyType;
 	}
 
-	REPORT_PYTHON_GENERATION_ISSUE(Fatal, TEXT("Failed to generate Python glue code for struct '%s'!"), *InStruct->GetName());
+	REPORT_PYTHON_GENERATION_ISSUE(Fatal, TEXT("Failed to generate Python glue code for struct '%s'!"), *InStruct->GetPathName());
 	return nullptr;
 }
 
-void FPyWrapperTypeRegistry::RegisterWrappedStructType(const FName StructName, PyTypeObject* PyType, const bool InDetectNameConflicts)
+void FPyWrapperTypeRegistry::RegisterWrappedStructType(const FSoftObjectPath& StructName, PyTypeObject* PyType, const bool InDetectNameConflicts)
 {
 	if (InDetectNameConflicts)
 	{
@@ -2143,7 +2176,7 @@ void FPyWrapperTypeRegistry::RegisterWrappedStructType(const FName StructName, P
 	PythonWrappedStructs.Add(StructName, PyType);
 }
 
-void FPyWrapperTypeRegistry::UnregisterWrappedStructType(const FName StructName, PyTypeObject* PyType, const bool InUnregisterName)
+void FPyWrapperTypeRegistry::UnregisterWrappedStructType(const FSoftObjectPath& StructName, PyTypeObject* PyType, const bool InUnregisterName)
 {
 	if (InUnregisterName)
 	{
@@ -2154,8 +2187,7 @@ void FPyWrapperTypeRegistry::UnregisterWrappedStructType(const FName StructName,
 
 bool FPyWrapperTypeRegistry::HasWrappedStructType(const UScriptStruct* InStruct) const
 {
-	const FName TypeRegistryName = PyGenUtil::GetTypeRegistryName(InStruct);
-	return PythonWrappedStructs.Contains(TypeRegistryName);
+	return PythonWrappedStructs.Contains(InStruct);
 }
 
 PyTypeObject* FPyWrapperTypeRegistry::GetWrappedStructType(const UScriptStruct* InStruct) const
@@ -2164,8 +2196,7 @@ PyTypeObject* FPyWrapperTypeRegistry::GetWrappedStructType(const UScriptStruct*
 
 	for (const UScriptStruct* Struct = InStruct; Struct; Struct = Cast<UScriptStruct>(Struct->GetSuperStruct()))
 	{
-		const FName TypeRegistryName = PyGenUtil::GetTypeRegistryName(Struct);
-		if (PyTypeObject* StructPyType = PythonWrappedStructs.FindRef(TypeRegistryName))
+		if (PyTypeObject* StructPyType = PythonWrappedStructs.FindRef(InStruct))
 		{
 			PyType = StructPyType;
 			break;
@@ -2183,7 +2214,7 @@ PyTypeObject* FPyWrapperTypeRegistry::GenerateWrappedEnumType(const UEnum* InEnu
 	bool bIsNewType = true;
 
 	// Already processed? Nothing more to do
-	const FName TypeRegistryName = PyGenUtil::GetTypeRegistryName(InEnum);
+	const FSoftObjectPath TypeRegistryName(InEnum);
 	if (PyTypeObject* ExistingPyType = PythonWrappedEnums.FindRef(TypeRegistryName))
 	{
 		if (!EnumHasAnyFlags(InGenerationFlags, EPyTypeGenerationFlags::OverwriteExisting))
@@ -2261,20 +2292,21 @@ PyTypeObject* FPyWrapperTypeRegistry::GenerateWrappedEnumType(const UEnum* InEnu
 			RegisterWrappedEnumType(TypeRegistryName, &GeneratedWrappedType->PyType, !bIsBlueprintGeneratedType);
 
 			// Also generate and register any deprecated aliases for this type
-			const TArray<FString> DeprecatedPythonEnumNames = PyGenUtil::GetDeprecatedEnumPythonNames(InEnum);
-			for (const FString& DeprecatedPythonEnumName : DeprecatedPythonEnumNames)
+			const TArray<TTuple<FSoftObjectPath, FString>> DeprecatedPythonEnumNames = PyGenUtil::GetDeprecatedEnumPythonNames(InEnum);
+			for (const TTuple<FSoftObjectPath, FString>& DeprecatedPythonEnumNamePair : DeprecatedPythonEnumNames)
 			{
-				const FName DeprecatedEnumName = *DeprecatedPythonEnumName;
+				const FString& DeprecatedPythonEnumName = DeprecatedPythonEnumNamePair.Value;
 				FString DeprecationMessage = FString::Printf(TEXT("'%s' was renamed to '%s'."), *DeprecatedPythonEnumName, *PythonEnumName);
 
-				if (GeneratedWrappedTypes.Contains(DeprecatedEnumName))
+				const FSoftObjectPath& DeprecatedTypeRegistryName = DeprecatedPythonEnumNamePair.Key;
+				if (GeneratedWrappedTypes.Contains(DeprecatedTypeRegistryName))
 				{
-					REPORT_PYTHON_GENERATION_ISSUE(Warning, TEXT("Deprecated enum name '%s' conflicted with an existing type!"), *DeprecatedPythonEnumName);
+					REPORT_PYTHON_GENERATION_ISSUE(Warning, TEXT("Deprecated enum name '%s' conflicted with an existing type!"), *DeprecatedTypeRegistryName.ToString());
 					continue;
 				}
 
 				TSharedRef<PyGenUtil::FGeneratedWrappedEnumType> DeprecatedGeneratedWrappedType = MakeShared<PyGenUtil::FGeneratedWrappedEnumType>();
-				GeneratedWrappedTypes.Add(DeprecatedEnumName, DeprecatedGeneratedWrappedType);
+				GeneratedWrappedTypes.Add(DeprecatedTypeRegistryName, DeprecatedGeneratedWrappedType);
 
 				DeprecatedGeneratedWrappedType->TypeName = PyGenUtil::TCHARToUTF8Buffer(*DeprecatedPythonEnumName);
 				DeprecatedGeneratedWrappedType->TypeDoc = PyGenUtil::TCHARToUTF8Buffer(*FString::Printf(TEXT("deprecated: %s"), *DeprecationMessage));
@@ -2291,7 +2323,7 @@ PyTypeObject* FPyWrapperTypeRegistry::GenerateWrappedEnumType(const UEnum* InEnu
 				{
 					if (!UnrealModuleName.IsNone())
 					{
-						GeneratedWrappedTypesForModule.Add(UnrealModuleName, DeprecatedEnumName);
+						GeneratedWrappedTypesForModule.Add(UnrealModuleName, DeprecatedTypeRegistryName);
 						// Execute Python code within this block
 						{
 							FPyScopedGIL GIL;
@@ -2301,11 +2333,11 @@ PyTypeObject* FPyWrapperTypeRegistry::GenerateWrappedEnumType(const UEnum* InEnu
 							PyModule_AddObject(PyModule, DeprecatedGeneratedWrappedType->PyType.tp_name, (PyObject*)&DeprecatedGeneratedWrappedType->PyType);
 						}
 					}
-					RegisterWrappedEnumType(DeprecatedEnumName, &DeprecatedGeneratedWrappedType->PyType, !bIsBlueprintGeneratedType);
+					RegisterWrappedEnumType(DeprecatedTypeRegistryName, &DeprecatedGeneratedWrappedType->PyType, !bIsBlueprintGeneratedType);
 				}
 				else
 				{
-					REPORT_PYTHON_GENERATION_ISSUE(Fatal, TEXT("Failed to generate Python glue code for deprecated enum '%s'!"), *DeprecatedPythonEnumName);
+					REPORT_PYTHON_GENERATION_ISSUE(Fatal, TEXT("Failed to generate Python glue code for deprecated enum '%s'!"), *DeprecatedTypeRegistryName.ToString());
 				}
 			}
 		}
@@ -2313,11 +2345,11 @@ PyTypeObject* FPyWrapperTypeRegistry::GenerateWrappedEnumType(const UEnum* InEnu
 		return &GeneratedWrappedType->PyType;
 	}
 
-	REPORT_PYTHON_GENERATION_ISSUE(Fatal, TEXT("Failed to generate Python glue code for enum '%s'!"), *InEnum->GetName());
+	REPORT_PYTHON_GENERATION_ISSUE(Fatal, TEXT("Failed to generate Python glue code for enum '%s'!"), *InEnum->GetPathName());
 	return nullptr;
 }
 
-void FPyWrapperTypeRegistry::RegisterWrappedEnumType(const FName EnumName, PyTypeObject* PyType, const bool InDetectNameConflicts)
+void FPyWrapperTypeRegistry::RegisterWrappedEnumType(const FSoftObjectPath& EnumName, PyTypeObject* PyType, const bool InDetectNameConflicts)
 {
 	if (InDetectNameConflicts)
 {
@@ -2326,7 +2358,7 @@ void FPyWrapperTypeRegistry::RegisterWrappedEnumType(const FName EnumName, PyTyp
 	PythonWrappedEnums.Add(EnumName, PyType);
 }
 
-void FPyWrapperTypeRegistry::UnregisterWrappedEnumType(const FName EnumName, PyTypeObject* PyType, const bool InUnregisterName)
+void FPyWrapperTypeRegistry::UnregisterWrappedEnumType(const FSoftObjectPath& EnumName, PyTypeObject* PyType, const bool InUnregisterName)
 {
 	if (InUnregisterName)
 	{
@@ -2337,16 +2369,14 @@ void FPyWrapperTypeRegistry::UnregisterWrappedEnumType(const FName EnumName, PyT
 
 bool FPyWrapperTypeRegistry::HasWrappedEnumType(const UEnum* InEnum) const
 {
-	const FName TypeRegistryName = PyGenUtil::GetTypeRegistryName(InEnum);
-	return PythonWrappedEnums.Contains(TypeRegistryName);
+	return PythonWrappedEnums.Contains(InEnum);
 }
 
 PyTypeObject* FPyWrapperTypeRegistry::GetWrappedEnumType(const UEnum* InEnum) const
 {
 	PyTypeObject* PyType = &PyWrapperEnumType;
 
-	const FName TypeRegistryName = PyGenUtil::GetTypeRegistryName(InEnum);
-	if (PyTypeObject* EnumPyType = PythonWrappedEnums.FindRef(TypeRegistryName))
+	if (PyTypeObject* EnumPyType = PythonWrappedEnums.FindRef(InEnum))
 	{
 		PyType = EnumPyType;
 	}
@@ -2360,7 +2390,7 @@ PyTypeObject* FPyWrapperTypeRegistry::GenerateWrappedDelegateType(const UFunctio
 	INC_DWORD_STAT(STAT_GenerateWrappedDelegateCallCount);
 
 	// Already processed? Nothing more to do
-	const FName TypeRegistryName = PyGenUtil::GetTypeRegistryName(InDelegateSignature);
+	const FSoftObjectPath TypeRegistryName(InDelegateSignature);
 	if (PyTypeObject* ExistingPyType = PythonWrappedDelegates.FindRef(TypeRegistryName))
 	{
 		return ExistingPyType;
@@ -2488,11 +2518,11 @@ PyTypeObject* FPyWrapperTypeRegistry::GenerateWrappedDelegateType(const UFunctio
 		return &GeneratedWrappedType->PyType;
 	}
 
-	REPORT_PYTHON_GENERATION_ISSUE(Fatal, TEXT("Failed to generate Python glue code for delegate '%s'!"), *InDelegateSignature->GetName());
+	REPORT_PYTHON_GENERATION_ISSUE(Fatal, TEXT("Failed to generate Python glue code for delegate '%s'!"), *InDelegateSignature->GetPathName());
 	return nullptr;
 }
 
-void FPyWrapperTypeRegistry::RegisterWrappedDelegateType(const FName DelegateName, PyTypeObject* PyType, const bool InDetectNameConflicts)
+void FPyWrapperTypeRegistry::RegisterWrappedDelegateType(const FSoftObjectPath& DelegateName, PyTypeObject* PyType, const bool InDetectNameConflicts)
 {
 	if (InDetectNameConflicts)
 	{
@@ -2501,7 +2531,7 @@ void FPyWrapperTypeRegistry::RegisterWrappedDelegateType(const FName DelegateNam
 	PythonWrappedDelegates.Add(DelegateName, PyType);
 }
 
-void FPyWrapperTypeRegistry::UnregisterWrappedDelegateType(const FName DelegateName, PyTypeObject* PyType, const bool InUnregisterName)
+void FPyWrapperTypeRegistry::UnregisterWrappedDelegateType(const FSoftObjectPath& DelegateName, PyTypeObject* PyType, const bool InUnregisterName)
 {
 	if (InUnregisterName)
 	{
@@ -2512,16 +2542,14 @@ void FPyWrapperTypeRegistry::UnregisterWrappedDelegateType(const FName DelegateN
 
 bool FPyWrapperTypeRegistry::HasWrappedDelegateType(const UFunction* InDelegateSignature) const
 {
-	const FName TypeRegistryName = PyGenUtil::GetTypeRegistryName(InDelegateSignature);
-	return PythonWrappedDelegates.Contains(TypeRegistryName);
+	return PythonWrappedDelegates.Contains(InDelegateSignature);
 }
 
 PyTypeObject* FPyWrapperTypeRegistry::GetWrappedDelegateType(const UFunction* InDelegateSignature) const
 {
 	PyTypeObject* PyType = InDelegateSignature->HasAnyFunctionFlags(FUNC_MulticastDelegate) ? &PyWrapperMulticastDelegateType : &PyWrapperDelegateType;
 
-	const FName TypeRegistryName = PyGenUtil::GetTypeRegistryName(InDelegateSignature);
-	if (PyTypeObject* DelegatePyType = PythonWrappedDelegates.FindRef(TypeRegistryName))
+	if (PyTypeObject* DelegatePyType = PythonWrappedDelegates.FindRef(InDelegateSignature))
 	{
 		PyType = DelegatePyType;
 	}
@@ -2533,7 +2561,7 @@ void FPyWrapperTypeRegistry::GatherWrappedTypesForPropertyReferences(const FProp
 {
 	if (const FObjectProperty* ObjProp = CastField<const FObjectProperty>(InProp))
 	{
-		if (ObjProp->PropertyClass && !PythonWrappedClasses.Contains(PyGenUtil::GetTypeRegistryName(ObjProp->PropertyClass)))
+		if (ObjProp->PropertyClass && !PythonWrappedClasses.Contains(ObjProp->PropertyClass))
 		{
 			OutGeneratedWrappedTypeReferences.ClassReferences.Add(ObjProp->PropertyClass);
 		}
@@ -2542,7 +2570,7 @@ void FPyWrapperTypeRegistry::GatherWrappedTypesForPropertyReferences(const FProp
 
 	if (const FStructProperty* StructProp = CastField<const FStructProperty>(InProp))
 	{
-		if (!PythonWrappedStructs.Contains(PyGenUtil::GetTypeRegistryName(StructProp->Struct)))
+		if (!PythonWrappedStructs.Contains(StructProp->Struct))
 		{
 			OutGeneratedWrappedTypeReferences.StructReferences.Add(StructProp->Struct);
 		}
@@ -2551,7 +2579,7 @@ void FPyWrapperTypeRegistry::GatherWrappedTypesForPropertyReferences(const FProp
 
 	if (const FEnumProperty* EnumProp = CastField<const FEnumProperty>(InProp))
 	{
-		if (!PythonWrappedStructs.Contains(PyGenUtil::GetTypeRegistryName(EnumProp->GetEnum())))
+		if (!PythonWrappedStructs.Contains(EnumProp->GetEnum()))
 		{
 			OutGeneratedWrappedTypeReferences.EnumReferences.Add(EnumProp->GetEnum());
 		}
@@ -2562,7 +2590,7 @@ void FPyWrapperTypeRegistry::GatherWrappedTypesForPropertyReferences(const FProp
 	{
 		if (ByteProp->Enum)
 		{
-			if (!PythonWrappedStructs.Contains(PyGenUtil::GetTypeRegistryName(ByteProp->Enum)))
+			if (!PythonWrappedStructs.Contains(ByteProp->Enum))
 			{
 				OutGeneratedWrappedTypeReferences.EnumReferences.Add(ByteProp->Enum);
 			}
@@ -2572,7 +2600,7 @@ void FPyWrapperTypeRegistry::GatherWrappedTypesForPropertyReferences(const FProp
 
 	if (const FDelegateProperty* DelegateProp = CastField<const FDelegateProperty>(InProp))
 	{
-		if (!PythonWrappedStructs.Contains(PyGenUtil::GetTypeRegistryName(DelegateProp->SignatureFunction)))
+		if (!PythonWrappedStructs.Contains(DelegateProp->SignatureFunction))
 		{
 			OutGeneratedWrappedTypeReferences.DelegateReferences.Add(DelegateProp->SignatureFunction);
 		}
@@ -2581,7 +2609,7 @@ void FPyWrapperTypeRegistry::GatherWrappedTypesForPropertyReferences(const FProp
 
 	if (const FMulticastDelegateProperty* DelegateProp = CastField<const FMulticastDelegateProperty>(InProp))
 	{
-		if (!PythonWrappedStructs.Contains(PyGenUtil::GetTypeRegistryName(DelegateProp->SignatureFunction)))
+		if (!PythonWrappedStructs.Contains(DelegateProp->SignatureFunction))
 		{
 			OutGeneratedWrappedTypeReferences.DelegateReferences.Add(DelegateProp->SignatureFunction);
 		}
@@ -2616,6 +2644,10 @@ void FPyWrapperTypeRegistry::GenerateStubCodeForWrappedTypes(const EPyOnlineDocs
 	FSlowHeartBeatScope SuspendHeartBeat;
 	FDisableHitchDetectorScope SuspendGameThreadHitch;
 
+	FScopedSlowTask SlowTask(1, LOCTEXT("GeneratingPythonStubFile", "Generating unreal.py..."));
+	SlowTask.Visibility = ESlowTaskVisibility::Important; // this function can be very slow, users will benefit from our messages
+	SlowTask.MakeDialog();
+
 	FPyScopedGIL GIL;
 	FPyFileWriter PythonScript;
 
@@ -2781,7 +2813,7 @@ void FPyWrapperTypeRegistry::GenerateStubCodeForWrappedTypes(const EPyOnlineDocs
 	// Process generated glue code
 	// Also excludes types that don't pass the filters specified in InDocGenFlags using the information about
 	// which module it came from and where that module exists on disk.
-	auto ProcessWrappedDataArray = [this, &PythonScript, InDocGenFlags](const TMap<FName, PyTypeObject*>& WrappedData, const TSharedPtr<FPyOnlineDocsSection>& OnlineDocsSection)
+	auto ProcessWrappedDataArray = [this, &PythonScript, InDocGenFlags](const TMap<FSoftObjectPath, PyTypeObject*>& WrappedData, const TSharedPtr<FPyOnlineDocsSection>& OnlineDocsSection)
 	{
 		if (OnlineDocsSection.IsValid())
 		{
@@ -3669,12 +3701,12 @@ void FPyWrapperTypeRegistry::GenerateStubCodeForWrappedType(PyTypeObject* PyType
 	OutPythonScript.WriteNewLine();
 }
 
-void FPyWrapperTypeRegistry::RegisterPythonTypeName(const FString& InPythonTypeName, const FName& InUnrealTypeName)
+void FPyWrapperTypeRegistry::RegisterPythonTypeName(const FString& InPythonTypeName, const FSoftObjectPath& InUnrealTypeName)
 {
-	const FName ExistingUnrealTypeName = PythonWrappedTypeNameToUnrealTypeName.FindRef(InPythonTypeName);
-	if (ExistingUnrealTypeName.IsNone())
+	FSoftObjectPath& ExistingUnrealTypeName = PythonWrappedTypeNameToUnrealTypeName.FindOrAdd(InPythonTypeName);
+	if (ExistingUnrealTypeName.IsNull())
 	{
-		PythonWrappedTypeNameToUnrealTypeName.Add(InPythonTypeName, InUnrealTypeName);
+		ExistingUnrealTypeName = InUnrealTypeName;
 	}
 	else
 	{
@@ -3682,9 +3714,9 @@ void FPyWrapperTypeRegistry::RegisterPythonTypeName(const FString& InPythonTypeN
 	}
 }
 
-void FPyWrapperTypeRegistry::UnregisterPythonTypeName(const FString& InPythonTypeName, const FName& InUnrealTypeName)
+void FPyWrapperTypeRegistry::UnregisterPythonTypeName(const FString& InPythonTypeName, const FSoftObjectPath& InUnrealTypeName)
 {
-	const FName* ExistingUnrealTypeNamePtr = PythonWrappedTypeNameToUnrealTypeName.Find(InPythonTypeName);
+	const FSoftObjectPath* ExistingUnrealTypeNamePtr = PythonWrappedTypeNameToUnrealTypeName.Find(InPythonTypeName);
 	if (ExistingUnrealTypeNamePtr && *ExistingUnrealTypeNamePtr == InUnrealTypeName)
 	{
 		PythonWrappedTypeNameToUnrealTypeName.Remove(InPythonTypeName);
@@ -3692,3 +3724,5 @@ void FPyWrapperTypeRegistry::UnregisterPythonTypeName(const FString& InPythonTyp
 }
 
 #endif	// WITH_PYTHON
+
+#undef LOCTEXT_NAMESPACE
diff --git a/Engine/Plugins/Experimental/PythonScriptPlugin/Source/PythonScriptPlugin/Private/PyWrapperTypeRegistry.h b/Engine/Plugins/Experimental/PythonScriptPlugin/Source/PythonScriptPlugin/Private/PyWrapperTypeRegistry.h
--- a/Engine/Plugins/Experimental/PythonScriptPlugin/Source/PythonScriptPlugin/Private/PyWrapperTypeRegistry.h
+++ b/Engine/Plugins/Experimental/PythonScriptPlugin/Source/PythonScriptPlugin/Private/PyWrapperTypeRegistry.h
@@ -381,7 +381,7 @@ public:
 	void RegisterInlineStructFactory(const TSharedRef<const IPyWrapperInlineStructFactory>& InFactory);
 
 	/** Get the factory for an inline struct (if known) from its Unreal struct name */
-	const IPyWrapperInlineStructFactory* GetInlineStructFactory(const FName StructName) const;
+	const IPyWrapperInlineStructFactory* GetInlineStructFactory(const FTopLevelAssetPath& StructName) const;
 
 	/** Generate a wrapped type for all currently available Unreal types */
 	void GenerateWrappedTypes();
@@ -399,10 +399,10 @@ public:
 	void NotifyModulesDirtied(const TSet<FName>& InDirtyModules) const;
 
 	/** Called when an asset that is registered with the type registry has been renamed */
-	void UpdateGenerateWrappedTypeForRename(const FName InOldTypeRegistryName, const UObject* InObj);
+	void UpdateGenerateWrappedTypeForRename(const FSoftObjectPath& InOldTypeRegistryName, const UObject* InObj);
 
 	/** Called when an asset that is registered with the type registry is deleted */
-	void RemoveGenerateWrappedTypeForDelete(const FName InTypeRegistryName);
+	void RemoveGenerateWrappedTypeForDelete(const FSoftObjectPath& InTypeRegistryName);
 
 	/** Generate a wrapped type for the given object (if it's a valid type to be wrapped) */
 	PyTypeObject* GenerateWrappedTypeForObject(const UObject* InObj, FGeneratedWrappedTypeReferences& OutGeneratedWrappedTypeReferences, TSet<FName>& OutDirtyModules, const EPyTypeGenerationFlags InGenerationFlags = EPyTypeGenerationFlags::None);
@@ -411,7 +411,7 @@ public:
 	bool HasWrappedTypeForObject(const UObject* InObj) const;
 
 	/** True if we have wrapped type for the given object name */
-	bool HasWrappedTypeForObjectName(const FName InName) const;
+	bool HasWrappedTypeForObjectName(const FSoftObjectPath& InName) const;
 
 	/** Get the best wrapped type for the given object */
 	PyTypeObject* GetWrappedTypeForObject(const UObject* InObj) const;
@@ -420,10 +420,10 @@ public:
 	PyTypeObject* GenerateWrappedClassType(const UClass* InClass, FGeneratedWrappedTypeReferences& OutGeneratedWrappedTypeReferences, TSet<FName>& OutDirtyModules, const EPyTypeGenerationFlags InGenerationFlags = EPyTypeGenerationFlags::None);
 
 	/** Register the wrapped type associated with the given class name */
-	void RegisterWrappedClassType(const FName ClassName, PyTypeObject* PyType, const bool InDetectNameConflicts = true);
+	void RegisterWrappedClassType(const FSoftObjectPath& ClassName, PyTypeObject* PyType, const bool InDetectNameConflicts = true);
 
 	/** Unregister the wrapped type associated with the given class name */
-	void UnregisterWrappedClassType(const FName ClassName, PyTypeObject* PyType, const bool InUnregisterName = true);
+	void UnregisterWrappedClassType(const FSoftObjectPath& ClassName, PyTypeObject* PyType, const bool InUnregisterName = true);
 
 	/** True if we have wrapped type for the exact given class */
 	bool HasWrappedClassType(const UClass* InClass) const;
@@ -435,10 +435,10 @@ public:
 	PyTypeObject* GenerateWrappedStructType(const UScriptStruct* InStruct, FGeneratedWrappedTypeReferences& OutGeneratedWrappedTypeReferences, TSet<FName>& OutDirtyModules, const EPyTypeGenerationFlags InGenerationFlags = EPyTypeGenerationFlags::None);
 
 	/** Register the wrapped type associated with the given struct name */
-	void RegisterWrappedStructType(const FName StructName, PyTypeObject* PyType, const bool InDetectNameConflicts = true);
+	void RegisterWrappedStructType(const FSoftObjectPath& StructName, PyTypeObject* PyType, const bool InDetectNameConflicts = true);
 
 	/** Unregister the wrapped type associated with the given struct name */
-	void UnregisterWrappedStructType(const FName StructName, PyTypeObject* PyType, const bool InUnregisterName = true);
+	void UnregisterWrappedStructType(const FSoftObjectPath& StructName, PyTypeObject* PyType, const bool InUnregisterName = true);
 
 	/** True if we have wrapped type for the exact given struct */
 	bool HasWrappedStructType(const UScriptStruct* InStruct) const;
@@ -450,10 +450,10 @@ public:
 	PyTypeObject* GenerateWrappedEnumType(const UEnum* InEnum, FGeneratedWrappedTypeReferences& OutGeneratedWrappedTypeReferences, TSet<FName>& OutDirtyModules, const EPyTypeGenerationFlags InGenerationFlags = EPyTypeGenerationFlags::None);
 
 	/** Register the wrapped type associated with the given enum name */
-	void RegisterWrappedEnumType(const FName EnumName, PyTypeObject* PyType, const bool InDetectNameConflicts = true);
+	void RegisterWrappedEnumType(const FSoftObjectPath& EnumName, PyTypeObject* PyType, const bool InDetectNameConflicts = true);
 
 	/** Unregister the wrapped type associated with the given enum name */
-	void UnregisterWrappedEnumType(const FName EnumName, PyTypeObject* PyType, const bool InUnregisterName = true);
+	void UnregisterWrappedEnumType(const FSoftObjectPath& EnumName, PyTypeObject* PyType, const bool InUnregisterName = true);
 
 	/** True if we have wrapped type for the exact given enum */
 	bool HasWrappedEnumType(const UEnum* InEnum) const;
@@ -465,10 +465,10 @@ public:
 	PyTypeObject* GenerateWrappedDelegateType(const UFunction* InDelegateSignature, FGeneratedWrappedTypeReferences& OutGeneratedWrappedTypeReferences, TSet<FName>& OutDirtyModules, const EPyTypeGenerationFlags InGenerationFlags = EPyTypeGenerationFlags::None);
 
 	/** Register the wrapped type associated with the given delegate name */
-	void RegisterWrappedDelegateType(const FName DelegateName, PyTypeObject* PyType, const bool InDetectNameConflicts = true);
+	void RegisterWrappedDelegateType(const FSoftObjectPath& DelegateName, PyTypeObject* PyType, const bool InDetectNameConflicts = true);
 
 	/** Unregister the wrapped type associated with the given delegate name */
-	void UnregisterWrappedDelegateType(const FName DelegateName, PyTypeObject* PyType, const bool InUnregisterName = true);
+	void UnregisterWrappedDelegateType(const FSoftObjectPath& DelegateName, PyTypeObject* PyType, const bool InUnregisterName = true);
 
 	/** True if we have wrapped type for the exact given delegate signature */
 	bool HasWrappedDelegateType(const UFunction* InDelegateSignature) const;
@@ -487,10 +487,10 @@ private:
 	static void GenerateStubCodeForWrappedType(PyTypeObject* PyType, const PyGenUtil::FGeneratedWrappedType* GeneratedTypeData, FPyFileWriter& OutPythonScript, FPyOnlineDocsSection* OutOnlineDocsSection = nullptr);
 
 	/** Register a Python type name, and detect if a name conflict has occurred */
-	void RegisterPythonTypeName(const FString& InPythonTypeName, const FName& InUnrealTypeName);
+	void RegisterPythonTypeName(const FString& InPythonTypeName, const FSoftObjectPath& InUnrealTypeName);
 
 	/** Unregister a Python type name if it corresponds to the given Unreal type name */
-	void UnregisterPythonTypeName(const FString& InPythonTypeName, const FName& InUnrealTypeName);
+	void UnregisterPythonTypeName(const FString& InPythonTypeName, const FSoftObjectPath& InUnrealTypeName);
 
 	FPyWrapperTypeRegistry();
 
@@ -498,29 +498,29 @@ private:
 	bool bCanRegisterInlineStructFactories;
 
 	/** Map from the Unreal struct name to the factory data for an inline struct (ie, a struct known at compile time that will allocate its instance data inlined within the Python object) */
-	TMap<FName, TSharedPtr<const IPyWrapperInlineStructFactory>> InlineStructFactories;
+	TMap<FTopLevelAssetPath, TSharedPtr<const IPyWrapperInlineStructFactory>> InlineStructFactories;
 
 	/** Map from the Python wrapped type name to the Unreal type name it was generated from (for conflict detection) */
-	typedef TMap<FString, FName, FDefaultSetAllocator, PyGenUtil::FCaseSensitiveStringMapFuncs<FName>> FCaseSensitiveStringToNameMap;
+	typedef TMap<FString, FSoftObjectPath, FDefaultSetAllocator, PyGenUtil::FCaseSensitiveStringMapFuncs<FSoftObjectPath>> FCaseSensitiveStringToNameMap;
 	FCaseSensitiveStringToNameMap PythonWrappedTypeNameToUnrealTypeName;
 
 	/** Map from the Unreal class name to the Python type */
-	TMap<FName, PyTypeObject*> PythonWrappedClasses;
+	TMap<FSoftObjectPath, PyTypeObject*> PythonWrappedClasses;
 
 	/** Map from the Unreal struct name to the Python type */
-	TMap<FName, PyTypeObject*> PythonWrappedStructs;
+	TMap<FSoftObjectPath, PyTypeObject*> PythonWrappedStructs;
 
 	/** Map from the Unreal enum name to the Python type */
-	TMap<FName, PyTypeObject*> PythonWrappedEnums;
+	TMap<FSoftObjectPath, PyTypeObject*> PythonWrappedEnums;
 
 	/** Map from the Unreal delegate signature name to the Python type */
-	TMap<FName, PyTypeObject*> PythonWrappedDelegates;
+	TMap<FSoftObjectPath, PyTypeObject*> PythonWrappedDelegates;
 
 	/** Map from the Unreal type name to the generated Python type data */
-	TMap<FName, TSharedPtr<PyGenUtil::FGeneratedWrappedType>> GeneratedWrappedTypes;
+	TMap<FSoftObjectPath, TSharedPtr<PyGenUtil::FGeneratedWrappedType>> GeneratedWrappedTypes;
 
 	/** Map from the Unreal module name to its generated type names (names are the Unreal names) */
-	TMultiMap<FName, FName> GeneratedWrappedTypesForModule;
+	TMultiMap<FName, FSoftObjectPath> GeneratedWrappedTypesForModule;
 
 	/** Map make func name to the meta data of the struct to make. Used when meta 'HasNativeMake' references a function not loaded yet. */
 	TMap<FString, TSharedPtr<FPyWrapperStructMetaData>> UnresolvedMakeFuncs;
diff --git a/Engine/Plugins/Experimental/PythonScriptPlugin/Source/PythonScriptPlugin/Private/PythonScriptPlugin.cpp b/Engine/Plugins/Experimental/PythonScriptPlugin/Source/PythonScriptPlugin/Private/PythonScriptPlugin.cpp
--- a/Engine/Plugins/Experimental/PythonScriptPlugin/Source/PythonScriptPlugin/Private/PythonScriptPlugin.cpp
+++ b/Engine/Plugins/Experimental/PythonScriptPlugin/Source/PythonScriptPlugin/Private/PythonScriptPlugin.cpp
@@ -68,6 +68,8 @@
 
 #define LOCTEXT_NAMESPACE "PythonScriptPlugin"
 
+#define UE_PYTHON_DEFER_INIT (1)
+
 #if WITH_PYTHON
 
 static PyUtil::FPyApiBuffer NullPyArg = PyUtil::TCHARToPyApiBuffer(TEXT(""));
@@ -678,26 +680,32 @@ void FPythonScriptPlugin::StartupModule()
 #if WITH_PYTHON
 	LLM_SCOPE_BYNAME(TEXT("PythonScriptPlugin"));
 
+#if !UE_PYTHON_DEFER_INIT
 	InitializePython();
+#endif
+
 	IModularFeatures::Get().RegisterModularFeature(IConsoleCommandExecutor::ModularFeatureName(), &CmdExec);
 	IModularFeatures::Get().RegisterModularFeature(IConsoleCommandExecutor::ModularFeatureName(), &CmdREPLExec);
 
 	check(!RemoteExecution);
 	RemoteExecution = MakeUnique<FPythonScriptRemoteExecution>(this);
 
-#if WITH_EDITOR
-	FPythonScriptPluginEditorStyle::Get();
 	FCoreDelegates::OnPostEngineInit.AddRaw(this, &FPythonScriptPlugin::OnPostEngineInit);
-#endif // WITH_EDITOR
-
 	FCoreDelegates::OnPreExit.AddRaw(this, &FPythonScriptPlugin::ShutdownPython);
 #endif	// WITH_PYTHON
 }
 
 #if WITH_PYTHON
-#if WITH_EDITOR
 void FPythonScriptPlugin::OnPostEngineInit()
 {
+	LLM_SCOPE_BYNAME(TEXT("PythonScriptPlugin"));
+
+#if UE_PYTHON_DEFER_INIT
+	InitializePython();
+#endif
+
+#if WITH_EDITOR
+	FPythonScriptPluginEditorStyle::Get();
 	if (UToolMenus::IsToolMenuUIEnabled())
 	{
 		check(CmdMenu == nullptr);
@@ -708,9 +716,9 @@ void FPythonScriptPlugin::OnPostEngineInit()
 			ExecPythonCommand(*InString);
 		}));
 	}
+#endif // WITH_EDITOR
 }
 
-#endif // WITH_EDITOR
 #endif // WITH_PYTHON
 
 void FPythonScriptPlugin::ShutdownModule()
@@ -722,12 +730,11 @@ void FPythonScriptPlugin::ShutdownModule()
 
 #if WITH_PYTHON
 	FCoreDelegates::OnPreExit.RemoveAll(this);
+	FCoreDelegates::OnPostEngineInit.RemoveAll(this);
 
 	RemoteExecution.Reset();
 
 #if WITH_EDITOR
-	FCoreDelegates::OnPostEngineInit.RemoveAll(this);
-
 	if (CmdMenu)
 	{
 		CmdMenu->OnShutdownMenu();
@@ -778,6 +785,10 @@ void FPythonScriptPlugin::InitializePython()
 
 	bInitialized = true;
 
+	FScopedSlowTask SlowTask(1, LOCTEXT("InitializingPython", "Initializing Python..."));
+	SlowTask.Visibility = ESlowTaskVisibility::Important; // this function can be very slow, users will benefit from our messages
+	SlowTask.MakeDialog();
+
 	const UPythonScriptPluginSettings* PythonPluginSettings = GetDefault<UPythonScriptPluginSettings>();
 
 	// Set-up the correct program name
@@ -1684,19 +1695,18 @@ void FPythonScriptPlugin::OnAssetRenamed(const FAssetData& Data, const FString&
 	TRACE_CPUPROFILER_EVENT_SCOPE(FPythonScriptPlugin::OnAssetRenamed)
 
 	FPyWrapperTypeRegistry& PyWrapperTypeRegistry = FPyWrapperTypeRegistry::Get();
-	const FName OldPackageName = *FPackageName::ObjectPathToPackageName(OldName);
 	
 	// If this asset has an associated Python type, then we need to rename it
-	if (PyWrapperTypeRegistry.HasWrappedTypeForObjectName(OldPackageName))
+	if (PyWrapperTypeRegistry.HasWrappedTypeForObjectName(OldName))
 	{
 		if (const UObject* AssetPtr = PyGenUtil::GetAssetTypeRegistryType(Data.GetAsset()))
 		{
-			PyWrapperTypeRegistry.UpdateGenerateWrappedTypeForRename(OldPackageName, AssetPtr);
+			PyWrapperTypeRegistry.UpdateGenerateWrappedTypeForRename(OldName, AssetPtr);
 			OnAssetUpdated(AssetPtr);
 		}
 		else
 		{
-			PyWrapperTypeRegistry.RemoveGenerateWrappedTypeForDelete(OldPackageName);
+			PyWrapperTypeRegistry.RemoveGenerateWrappedTypeForDelete(OldName);
 		}
 	}
 }
@@ -1708,9 +1718,10 @@ void FPythonScriptPlugin::OnAssetRemoved(const FAssetData& Data)
 	FPyWrapperTypeRegistry& PyWrapperTypeRegistry = FPyWrapperTypeRegistry::Get();
 	
 	// If this asset has an associated Python type, then we need to remove it
-	if (PyWrapperTypeRegistry.HasWrappedTypeForObjectName(Data.PackageName))
+	const FSoftObjectPath AssetPath = Data.GetSoftObjectPath();
+	if (PyWrapperTypeRegistry.HasWrappedTypeForObjectName(AssetPath))
 	{
-		PyWrapperTypeRegistry.RemoveGenerateWrappedTypeForDelete(Data.PackageName);
+		PyWrapperTypeRegistry.RemoveGenerateWrappedTypeForDelete(AssetPath);
 	}
 }
 
@@ -1810,6 +1821,8 @@ void FPythonScriptPlugin::PopulatePythonFileContextMenu(UToolMenu* InMenu)
 
 #endif	// WITH_PYTHON
 
+#undef UE_PYTHON_DEFER_INIT
+
 IMPLEMENT_MODULE(FPythonScriptPlugin, PythonScriptPlugin)
 
 #undef LOCTEXT_NAMESPACE
diff --git a/Engine/Plugins/Experimental/PythonScriptPlugin/Source/PythonScriptPlugin/Private/PythonScriptPlugin.h b/Engine/Plugins/Experimental/PythonScriptPlugin/Source/PythonScriptPlugin/Private/PythonScriptPlugin.h
--- a/Engine/Plugins/Experimental/PythonScriptPlugin/Source/PythonScriptPlugin/Private/PythonScriptPlugin.h
+++ b/Engine/Plugins/Experimental/PythonScriptPlugin/Source/PythonScriptPlugin/Private/PythonScriptPlugin.h
@@ -123,9 +123,7 @@ public:
 #if WITH_PYTHON
 	void RunStartupScripts();
 
-#if WITH_EDITOR	
 	void OnPostEngineInit();
-#endif
 
 	/** Sync the remote execution environment to the current settings, starting or stopping it as required */
 	void SyncRemoteExecutionToSettings();
