Fixed Python generated types being GC'd when running a GC without RF_Standalone

This splits the Python generated type meta-data reference collection into two; one for instances (the old ARO function), and one for the type itself.
This allows us to collect these type-specific references without needing a Python instance of that type.

diff --git a/Engine/Plugins/Experimental/PythonScriptPlugin/Source/PythonScriptPlugin/Private/PyReferenceCollector.cpp b/Engine/Plugins/Experimental/PythonScriptPlugin/Source/PythonScriptPlugin/Private/PyReferenceCollector.cpp
--- a/Engine/Plugins/Experimental/PythonScriptPlugin/Source/PythonScriptPlugin/Private/PyReferenceCollector.cpp
+++ b/Engine/Plugins/Experimental/PythonScriptPlugin/Source/PythonScriptPlugin/Private/PyReferenceCollector.cpp
@@ -39,12 +39,10 @@ void FPyReferenceCollector::AddReferencedObjects(FReferenceCollector& InCollecto
 			FPyWrapperBaseMetaData* PythonWrappedInstanceMetaData = FPyWrapperBaseMetaData::GetMetaData(PythonWrappedInstance);
 			if (PythonWrappedInstanceMetaData)
 			{
-				PythonWrappedInstanceMetaData->AddReferencedObjects(PythonWrappedInstance, InCollector);
+				PythonWrappedInstanceMetaData->AddInstanceReferencedObjects(PythonWrappedInstance, InCollector);
 			}
 		}
 	}
-
-	FPyWrapperTypeReinstancer::Get().AddReferencedObjects(InCollector);
 }
 
 FString FPyReferenceCollector::GetReferencerName() const
diff --git a/Engine/Plugins/Experimental/PythonScriptPlugin/Source/PythonScriptPlugin/Private/PyWrapperArray.cpp b/Engine/Plugins/Experimental/PythonScriptPlugin/Source/PythonScriptPlugin/Private/PyWrapperArray.cpp
--- a/Engine/Plugins/Experimental/PythonScriptPlugin/Source/PythonScriptPlugin/Private/PyWrapperArray.cpp
+++ b/Engine/Plugins/Experimental/PythonScriptPlugin/Source/PythonScriptPlugin/Private/PyWrapperArray.cpp
@@ -1517,7 +1517,7 @@ PyTypeObject InitializePyWrapperArrayIteratorType()
 PyTypeObject PyWrapperArrayType = InitializePyWrapperArrayType();
 PyTypeObject PyWrapperArrayIteratorType = InitializePyWrapperArrayIteratorType();
 
-void FPyWrapperArrayMetaData::AddReferencedObjects(FPyWrapperBase* Instance, FReferenceCollector& Collector)
+void FPyWrapperArrayMetaData::AddInstanceReferencedObjects(FPyWrapperBase* Instance, FReferenceCollector& Collector)
 {
 	FPyWrapperArray* Self = static_cast<FPyWrapperArray*>(Instance);
 	if (Self->ArrayProp && Self->ArrayInstance && !Self->OwnerContext.HasOwner())
diff --git a/Engine/Plugins/Experimental/PythonScriptPlugin/Source/PythonScriptPlugin/Private/PyWrapperArray.h b/Engine/Plugins/Experimental/PythonScriptPlugin/Source/PythonScriptPlugin/Private/PyWrapperArray.h
--- a/Engine/Plugins/Experimental/PythonScriptPlugin/Source/PythonScriptPlugin/Private/PyWrapperArray.h
+++ b/Engine/Plugins/Experimental/PythonScriptPlugin/Source/PythonScriptPlugin/Private/PyWrapperArray.h
@@ -112,7 +112,7 @@ struct FPyWrapperArrayMetaData : public FPyWrapperBaseMetaData
 	}
 
 	/** Add object references from the given Python object to the given collector */
-	virtual void AddReferencedObjects(FPyWrapperBase* Instance, FReferenceCollector& Collector) override;
+	virtual void AddInstanceReferencedObjects(FPyWrapperBase* Instance, FReferenceCollector& Collector) override;
 };
 
 typedef TPyPtr<FPyWrapperArray> FPyWrapperArrayPtr;
diff --git a/Engine/Plugins/Experimental/PythonScriptPlugin/Source/PythonScriptPlugin/Private/PyWrapperBase.h b/Engine/Plugins/Experimental/PythonScriptPlugin/Source/PythonScriptPlugin/Private/PyWrapperBase.h
--- a/Engine/Plugins/Experimental/PythonScriptPlugin/Source/PythonScriptPlugin/Private/PyWrapperBase.h
+++ b/Engine/Plugins/Experimental/PythonScriptPlugin/Source/PythonScriptPlugin/Private/PyWrapperBase.h
@@ -78,8 +78,13 @@ struct FPyWrapperBaseMetaData
 		return nullptr;
 	}
 
+	/** Add object references from this type meta-data to the given collector */
+	virtual void AddTypeReferencedObjects(FReferenceCollector& Collector)
+	{
+	}
+
 	/** Add object references from the given Python object to the given collector */
-	virtual void AddReferencedObjects(FPyWrapperBase* Instance, FReferenceCollector& Collector)
+	virtual void AddInstanceReferencedObjects(FPyWrapperBase* Instance, FReferenceCollector& Collector)
 	{
 	}
 };
diff --git a/Engine/Plugins/Experimental/PythonScriptPlugin/Source/PythonScriptPlugin/Private/PyWrapperDelegate.cpp b/Engine/Plugins/Experimental/PythonScriptPlugin/Source/PythonScriptPlugin/Private/PyWrapperDelegate.cpp
--- a/Engine/Plugins/Experimental/PythonScriptPlugin/Source/PythonScriptPlugin/Private/PyWrapperDelegate.cpp
+++ b/Engine/Plugins/Experimental/PythonScriptPlugin/Source/PythonScriptPlugin/Private/PyWrapperDelegate.cpp
@@ -1246,7 +1246,7 @@ PyTypeObject InitializePyWrapperMulticastDelegateType()
 PyTypeObject PyWrapperDelegateType = InitializePyWrapperDelegateType();
 PyTypeObject PyWrapperMulticastDelegateType = InitializePyWrapperMulticastDelegateType();
 
-void FPyWrapperDelegateMetaData::AddReferencedObjects(FPyWrapperBase* Instance, FReferenceCollector& Collector)
+void FPyWrapperDelegateMetaData::AddInstanceReferencedObjects(FPyWrapperBase* Instance, FReferenceCollector& Collector)
 {
 	FPyWrapperDelegate* Self = static_cast<FPyWrapperDelegate*>(Instance);
 	if (Self->DelegateInstance)
@@ -1255,7 +1255,7 @@ void FPyWrapperDelegateMetaData::AddReferencedObjects(FPyWrapperBase* Instance,
 	}
 }
 
-void FPyWrapperMulticastDelegateMetaData::AddReferencedObjects(FPyWrapperBase* Instance, FReferenceCollector& Collector)
+void FPyWrapperMulticastDelegateMetaData::AddInstanceReferencedObjects(FPyWrapperBase* Instance, FReferenceCollector& Collector)
 {
 	FPyWrapperMulticastDelegate* Self = static_cast<FPyWrapperMulticastDelegate*>(Instance);
 	if (Self->DelegateInstance)
diff --git a/Engine/Plugins/Experimental/PythonScriptPlugin/Source/PythonScriptPlugin/Private/PyWrapperDelegate.h b/Engine/Plugins/Experimental/PythonScriptPlugin/Source/PythonScriptPlugin/Private/PyWrapperDelegate.h
--- a/Engine/Plugins/Experimental/PythonScriptPlugin/Source/PythonScriptPlugin/Private/PyWrapperDelegate.h
+++ b/Engine/Plugins/Experimental/PythonScriptPlugin/Source/PythonScriptPlugin/Private/PyWrapperDelegate.h
@@ -115,8 +115,8 @@ struct TPyWrapperDelegateMetaData : public FPyWrapperBaseMetaData
 		return GetPythonCallableForDelegateClass(Py_TYPE(Instance));
 	}
 
-	/** Add object references from the given Python object to the given collector */
-	virtual void AddReferencedObjects(FPyWrapperBase* Instance, FReferenceCollector& Collector) override
+	/** Add object references from this type meta-data to the given collector */
+	virtual void AddTypeReferencedObjects(FReferenceCollector& Collector) override
 	{
 		Collector.AddReferencedObject(DelegateSignature.Func);
 		Collector.AddReferencedObject(PythonCallableForDelegateClass);
@@ -172,7 +172,7 @@ struct FPyWrapperDelegateMetaData : public TPyWrapperDelegateMetaData<FPyWrapper
 	PY_METADATA_METHODS(FPyWrapperDelegateMetaData, FGuid(0xCB3D0485, 0x8A3A443E, 0xBEE336F4, 0x82888A81))
 
 	/** Add object references from the given Python object to the given collector */
-	virtual void AddReferencedObjects(FPyWrapperBase* Instance, FReferenceCollector& Collector) override;
+	virtual void AddInstanceReferencedObjects(FPyWrapperBase* Instance, FReferenceCollector& Collector) override;
 };
 
 /** Type for all Unreal exposed multicast delegate instances */
@@ -212,7 +212,7 @@ struct FPyWrapperMulticastDelegateMetaData : public TPyWrapperDelegateMetaData<F
 	PY_METADATA_METHODS(FPyWrapperMulticastDelegateMetaData, FGuid(0x448FB4DA, 0x38DC4386, 0xBCAFF448, 0x29C0F3A4))
 
 	/** Add object references from the given Python object to the given collector */
-	virtual void AddReferencedObjects(FPyWrapperBase* Instance, FReferenceCollector& Collector) override;
+	virtual void AddInstanceReferencedObjects(FPyWrapperBase* Instance, FReferenceCollector& Collector) override;
 };
 
 typedef TPyPtr<FPyWrapperDelegate> FPyWrapperDelegatePtr;
diff --git a/Engine/Plugins/Experimental/PythonScriptPlugin/Source/PythonScriptPlugin/Private/PyWrapperEnum.cpp b/Engine/Plugins/Experimental/PythonScriptPlugin/Source/PythonScriptPlugin/Private/PyWrapperEnum.cpp
--- a/Engine/Plugins/Experimental/PythonScriptPlugin/Source/PythonScriptPlugin/Private/PyWrapperEnum.cpp
+++ b/Engine/Plugins/Experimental/PythonScriptPlugin/Source/PythonScriptPlugin/Private/PyWrapperEnum.cpp
@@ -733,7 +733,7 @@ bool FPyWrapperEnumMetaData::IsEnumFinalized(FPyWrapperEnum* Instance)
 	return IsEnumFinalized(Py_TYPE(Instance));
 }
 
-void FPyWrapperEnumMetaData::AddReferencedObjects(FPyWrapperBase* Instance, FReferenceCollector& Collector)
+void FPyWrapperEnumMetaData::AddTypeReferencedObjects(FReferenceCollector& Collector)
 {
 	Collector.AddReferencedObject(Enum);
 }
diff --git a/Engine/Plugins/Experimental/PythonScriptPlugin/Source/PythonScriptPlugin/Private/PyWrapperEnum.h b/Engine/Plugins/Experimental/PythonScriptPlugin/Source/PythonScriptPlugin/Private/PyWrapperEnum.h
--- a/Engine/Plugins/Experimental/PythonScriptPlugin/Source/PythonScriptPlugin/Private/PyWrapperEnum.h
+++ b/Engine/Plugins/Experimental/PythonScriptPlugin/Source/PythonScriptPlugin/Private/PyWrapperEnum.h
@@ -85,8 +85,8 @@ struct FPyWrapperEnumMetaData : public FPyWrapperBaseMetaData
 	/** Check to see if the enum is finalized */
 	static bool IsEnumFinalized(FPyWrapperEnum* Instance);
 
-	/** Add object references from the given Python object to the given collector */
-	virtual void AddReferencedObjects(FPyWrapperBase* Instance, FReferenceCollector& Collector) override;
+	/** Add object references from this type meta-data to the given collector */
+	virtual void AddTypeReferencedObjects(FReferenceCollector& Collector) override;
 
 	/** Get the reflection meta data type object associated with this wrapper type if there is one or nullptr if not. */
 	virtual const UField* GetMetaType() const override
diff --git a/Engine/Plugins/Experimental/PythonScriptPlugin/Source/PythonScriptPlugin/Private/PyWrapperFixedArray.cpp b/Engine/Plugins/Experimental/PythonScriptPlugin/Source/PythonScriptPlugin/Private/PyWrapperFixedArray.cpp
--- a/Engine/Plugins/Experimental/PythonScriptPlugin/Source/PythonScriptPlugin/Private/PyWrapperFixedArray.cpp
+++ b/Engine/Plugins/Experimental/PythonScriptPlugin/Source/PythonScriptPlugin/Private/PyWrapperFixedArray.cpp
@@ -910,7 +910,7 @@ PyTypeObject InitializePyWrapperFixedArrayIteratorType()
 PyTypeObject PyWrapperFixedArrayType = InitializePyWrapperFixedArrayType();
 PyTypeObject PyWrapperFixedArrayIteratorType = InitializePyWrapperFixedArrayIteratorType();
 
-void FPyWrapperFixedArrayMetaData::AddReferencedObjects(FPyWrapperBase* Instance, FReferenceCollector& Collector)
+void FPyWrapperFixedArrayMetaData::AddInstanceReferencedObjects(FPyWrapperBase* Instance, FReferenceCollector& Collector)
 {
 	FPyWrapperFixedArray* Self = static_cast<FPyWrapperFixedArray*>(Instance);
 	if (Self->ArrayProp && Self->ArrayInstance && !Self->OwnerContext.HasOwner())
diff --git a/Engine/Plugins/Experimental/PythonScriptPlugin/Source/PythonScriptPlugin/Private/PyWrapperFixedArray.h b/Engine/Plugins/Experimental/PythonScriptPlugin/Source/PythonScriptPlugin/Private/PyWrapperFixedArray.h
--- a/Engine/Plugins/Experimental/PythonScriptPlugin/Source/PythonScriptPlugin/Private/PyWrapperFixedArray.h
+++ b/Engine/Plugins/Experimental/PythonScriptPlugin/Source/PythonScriptPlugin/Private/PyWrapperFixedArray.h
@@ -82,7 +82,7 @@ struct FPyWrapperFixedArrayMetaData : public FPyWrapperBaseMetaData
 	}
 
 	/** Add object references from the given Python object to the given collector */
-	virtual void AddReferencedObjects(FPyWrapperBase* Instance, FReferenceCollector& Collector) override;
+	virtual void AddInstanceReferencedObjects(FPyWrapperBase* Instance, FReferenceCollector& Collector) override;
 };
 
 typedef TPyPtr<FPyWrapperFixedArray> FPyWrapperFixedArrayPtr;
diff --git a/Engine/Plugins/Experimental/PythonScriptPlugin/Source/PythonScriptPlugin/Private/PyWrapperMap.cpp b/Engine/Plugins/Experimental/PythonScriptPlugin/Source/PythonScriptPlugin/Private/PyWrapperMap.cpp
--- a/Engine/Plugins/Experimental/PythonScriptPlugin/Source/PythonScriptPlugin/Private/PyWrapperMap.cpp
+++ b/Engine/Plugins/Experimental/PythonScriptPlugin/Source/PythonScriptPlugin/Private/PyWrapperMap.cpp
@@ -1733,7 +1733,7 @@ PyTypeObject PyWrapperMapItemViewType = InitializePyWrapperMapSetViewType<FPyWra
 PyTypeObject PyWrapperMapKeyViewType = InitializePyWrapperMapSetViewType<FPyWrapperMapKeyView>("MapKeyView");
 PyTypeObject PyWrapperMapValueViewType = InitializePyWrapperMapViewType<FPyWrapperMapValueView>("MapValueView");
 
-void FPyWrapperMapMetaData::AddReferencedObjects(FPyWrapperBase* Instance, FReferenceCollector& Collector)
+void FPyWrapperMapMetaData::AddInstanceReferencedObjects(FPyWrapperBase* Instance, FReferenceCollector& Collector)
 {
 	FPyWrapperMap* Self = static_cast<FPyWrapperMap*>(Instance);
 	if (Self->MapProp && Self->MapInstance && !Self->OwnerContext.HasOwner())
diff --git a/Engine/Plugins/Experimental/PythonScriptPlugin/Source/PythonScriptPlugin/Private/PyWrapperMap.h b/Engine/Plugins/Experimental/PythonScriptPlugin/Source/PythonScriptPlugin/Private/PyWrapperMap.h
--- a/Engine/Plugins/Experimental/PythonScriptPlugin/Source/PythonScriptPlugin/Private/PyWrapperMap.h
+++ b/Engine/Plugins/Experimental/PythonScriptPlugin/Source/PythonScriptPlugin/Private/PyWrapperMap.h
@@ -103,7 +103,7 @@ struct FPyWrapperMapMetaData : public FPyWrapperBaseMetaData
 	}
 
 	/** Add object references from the given Python object to the given collector */
-	virtual void AddReferencedObjects(FPyWrapperBase* Instance, FReferenceCollector& Collector) override;
+	virtual void AddInstanceReferencedObjects(FPyWrapperBase* Instance, FReferenceCollector& Collector) override;
 };
 
 typedef TPyPtr<FPyWrapperMap> FPyWrapperMapPtr;
diff --git a/Engine/Plugins/Experimental/PythonScriptPlugin/Source/PythonScriptPlugin/Private/PyWrapperObject.cpp b/Engine/Plugins/Experimental/PythonScriptPlugin/Source/PythonScriptPlugin/Private/PyWrapperObject.cpp
--- a/Engine/Plugins/Experimental/PythonScriptPlugin/Source/PythonScriptPlugin/Private/PyWrapperObject.cpp
+++ b/Engine/Plugins/Experimental/PythonScriptPlugin/Source/PythonScriptPlugin/Private/PyWrapperObject.cpp
@@ -1130,11 +1130,14 @@ PyTypeObject InitializePyWrapperObjectType()
 
 PyTypeObject PyWrapperObjectType = InitializePyWrapperObjectType();
 
-void FPyWrapperObjectMetaData::AddReferencedObjects(FPyWrapperBase* Instance, FReferenceCollector& Collector)
+void FPyWrapperObjectMetaData::AddTypeReferencedObjects(FReferenceCollector& Collector)
 {
-	FPyWrapperObject* Self = static_cast<FPyWrapperObject*>(Instance);
-
 	Collector.AddReferencedObject(Class);
+}
+
+void FPyWrapperObjectMetaData::AddInstanceReferencedObjects(FPyWrapperBase* Instance, FReferenceCollector& Collector)
+{
+	FPyWrapperObject* Self = static_cast<FPyWrapperObject*>(Instance);
 
 	UObject* OldInstance = Self->ObjectInstance;
 	Collector.AddReferencedObject(Self->ObjectInstance);
diff --git a/Engine/Plugins/Experimental/PythonScriptPlugin/Source/PythonScriptPlugin/Private/PyWrapperObject.h b/Engine/Plugins/Experimental/PythonScriptPlugin/Source/PythonScriptPlugin/Private/PyWrapperObject.h
--- a/Engine/Plugins/Experimental/PythonScriptPlugin/Source/PythonScriptPlugin/Private/PyWrapperObject.h
+++ b/Engine/Plugins/Experimental/PythonScriptPlugin/Source/PythonScriptPlugin/Private/PyWrapperObject.h
@@ -147,8 +147,11 @@ struct FPyWrapperObjectMetaData : public FPyWrapperBaseMetaData
 	/** Check to see if the class is deprecated, and optionally return its deprecation message */
 	static bool IsClassDeprecated(FPyWrapperObject* Instance, FString* OutDeprecationMessage = nullptr);
 
+	/** Add object references from this type meta-data to the given collector */
+	virtual void AddTypeReferencedObjects(FReferenceCollector& Collector) override;
+
 	/** Add object references from the given Python object to the given collector */
-	virtual void AddReferencedObjects(FPyWrapperBase* Instance, FReferenceCollector& Collector) override;
+	virtual void AddInstanceReferencedObjects(FPyWrapperBase* Instance, FReferenceCollector& Collector) override;
 
 	/** Get the reflection meta data type object associated with this wrapper type if there is one or nullptr if not. */
 	virtual const UField* GetMetaType() const override
diff --git a/Engine/Plugins/Experimental/PythonScriptPlugin/Source/PythonScriptPlugin/Private/PyWrapperSet.cpp b/Engine/Plugins/Experimental/PythonScriptPlugin/Source/PythonScriptPlugin/Private/PyWrapperSet.cpp
--- a/Engine/Plugins/Experimental/PythonScriptPlugin/Source/PythonScriptPlugin/Private/PyWrapperSet.cpp
+++ b/Engine/Plugins/Experimental/PythonScriptPlugin/Source/PythonScriptPlugin/Private/PyWrapperSet.cpp
@@ -1548,7 +1548,7 @@ PyTypeObject InitializePyWrapperSetIteratorType()
 PyTypeObject PyWrapperSetType = InitializePyWrapperSetType();
 PyTypeObject PyWrapperSetIteratorType = InitializePyWrapperSetIteratorType();
 
-void FPyWrapperSetMetaData::AddReferencedObjects(FPyWrapperBase* Instance, FReferenceCollector& Collector)
+void FPyWrapperSetMetaData::AddInstanceReferencedObjects(FPyWrapperBase* Instance, FReferenceCollector& Collector)
 {
 	FPyWrapperSet* Self = static_cast<FPyWrapperSet*>(Instance);
 	if (Self->SetProp && Self->SetInstance && !Self->OwnerContext.HasOwner())
diff --git a/Engine/Plugins/Experimental/PythonScriptPlugin/Source/PythonScriptPlugin/Private/PyWrapperSet.h b/Engine/Plugins/Experimental/PythonScriptPlugin/Source/PythonScriptPlugin/Private/PyWrapperSet.h
--- a/Engine/Plugins/Experimental/PythonScriptPlugin/Source/PythonScriptPlugin/Private/PyWrapperSet.h
+++ b/Engine/Plugins/Experimental/PythonScriptPlugin/Source/PythonScriptPlugin/Private/PyWrapperSet.h
@@ -115,7 +115,7 @@ struct FPyWrapperSetMetaData : public FPyWrapperBaseMetaData
 	}
 
 	/** Add object references from the given Python object to the given collector */
-	virtual void AddReferencedObjects(FPyWrapperBase* Instance, FReferenceCollector& Collector) override;
+	virtual void AddInstanceReferencedObjects(FPyWrapperBase* Instance, FReferenceCollector& Collector) override;
 };
 
 typedef TPyPtr<FPyWrapperSet> FPyWrapperSetPtr;
diff --git a/Engine/Plugins/Experimental/PythonScriptPlugin/Source/PythonScriptPlugin/Private/PyWrapperStruct.cpp b/Engine/Plugins/Experimental/PythonScriptPlugin/Source/PythonScriptPlugin/Private/PyWrapperStruct.cpp
--- a/Engine/Plugins/Experimental/PythonScriptPlugin/Source/PythonScriptPlugin/Private/PyWrapperStruct.cpp
+++ b/Engine/Plugins/Experimental/PythonScriptPlugin/Source/PythonScriptPlugin/Private/PyWrapperStruct.cpp
@@ -1335,11 +1335,14 @@ FPyWrapperStructMetaData::FPyWrapperStructMetaData()
 {
 }
 
-void FPyWrapperStructMetaData::AddReferencedObjects(FPyWrapperBase* Instance, FReferenceCollector& Collector)
+void FPyWrapperStructMetaData::AddTypeReferencedObjects(FReferenceCollector& Collector)
 {
-	FPyWrapperStruct* Self = static_cast<FPyWrapperStruct*>(Instance);
-
 	Collector.AddReferencedObject(Struct);
+}
+
+void FPyWrapperStructMetaData::AddInstanceReferencedObjects(FPyWrapperBase* Instance, FReferenceCollector& Collector)
+{
+	FPyWrapperStruct* Self = static_cast<FPyWrapperStruct*>(Instance);
 	
 	Collector.AddReferencedObject(Self->ScriptStruct);
 	if (Self->ScriptStruct && Self->StructInstance && !Self->OwnerContext.HasOwner())
diff --git a/Engine/Plugins/Experimental/PythonScriptPlugin/Source/PythonScriptPlugin/Private/PyWrapperStruct.h b/Engine/Plugins/Experimental/PythonScriptPlugin/Source/PythonScriptPlugin/Private/PyWrapperStruct.h
--- a/Engine/Plugins/Experimental/PythonScriptPlugin/Source/PythonScriptPlugin/Private/PyWrapperStruct.h
+++ b/Engine/Plugins/Experimental/PythonScriptPlugin/Source/PythonScriptPlugin/Private/PyWrapperStruct.h
@@ -329,8 +329,11 @@ struct FPyWrapperStructMetaData : public FPyWrapperBaseMetaData
 	/** Check to see if the struct is deprecated, and optionally return its deprecation message */
 	static bool IsStructDeprecated(FPyWrapperStruct* Instance, FString* OutDeprecationMessage = nullptr);
 
+	/** Add object references from this type meta-data to the given collector */
+	virtual void AddTypeReferencedObjects(FReferenceCollector& Collector) override;
+
 	/** Add object references from the given Python object to the given collector */
-	virtual void AddReferencedObjects(FPyWrapperBase* Instance, FReferenceCollector& Collector) override;
+	virtual void AddInstanceReferencedObjects(FPyWrapperBase* Instance, FReferenceCollector& Collector) override;
 
 	/** Get the reflection meta data type object associated with this wrapper type if there is one or nullptr if not. */
 	virtual const UField* GetMetaType() const override
diff --git a/Engine/Plugins/Experimental/PythonScriptPlugin/Source/PythonScriptPlugin/Private/PyWrapperTypeRegistry.cpp b/Engine/Plugins/Experimental/PythonScriptPlugin/Source/PythonScriptPlugin/Private/PyWrapperTypeRegistry.cpp
--- a/Engine/Plugins/Experimental/PythonScriptPlugin/Source/PythonScriptPlugin/Private/PyWrapperTypeRegistry.cpp
+++ b/Engine/Plugins/Experimental/PythonScriptPlugin/Source/PythonScriptPlugin/Private/PyWrapperTypeRegistry.cpp
@@ -814,6 +814,11 @@ void FPyWrapperTypeReinstancer::AddReferencedObjects(FReferenceCollector& InColl
 	}
 }
 
+FString FPyWrapperTypeReinstancer::GetReferencerName() const
+{
+	return TEXT("FPyWrapperTypeReinstancer");
+}
+
 
 FPyWrapperTypeRegistry::FPyWrapperTypeRegistry()
 	: bCanRegisterInlineStructFactories(true)
@@ -838,6 +843,30 @@ FPyWrapperTypeRegistry& FPyWrapperTypeRegistry::Get()
 	return Instance;
 }
 
+void FPyWrapperTypeRegistry::AddReferencedObjects(FReferenceCollector& InCollector)
+{
+	auto AddTypeReferencedObjects = [&InCollector](const auto& PythonWrappedTypes)
+	{
+		for (const auto& PythonWrappedTypePair : PythonWrappedTypes)
+		{
+			if (FPyWrapperBaseMetaData* PythonWrappedTypeMetaData = FPyWrapperBaseMetaData::GetMetaData(PythonWrappedTypePair.Value))
+			{
+				PythonWrappedTypeMetaData->AddTypeReferencedObjects(InCollector);
+			}
+		}
+	};
+
+	AddTypeReferencedObjects(PythonWrappedClasses);
+	AddTypeReferencedObjects(PythonWrappedStructs);
+	AddTypeReferencedObjects(PythonWrappedEnums);
+	AddTypeReferencedObjects(PythonWrappedDelegates);
+}
+
+FString FPyWrapperTypeRegistry::GetReferencerName() const
+{
+	return TEXT("FPyWrapperTypeRegistry");
+}
+
 void FPyWrapperTypeRegistry::RegisterNativePythonModule(PyGenUtil::FNativePythonModule&& NativePythonModule)
 {
 	NativePythonModules.Add(MoveTemp(NativePythonModule));
diff --git a/Engine/Plugins/Experimental/PythonScriptPlugin/Source/PythonScriptPlugin/Private/PyWrapperTypeRegistry.h b/Engine/Plugins/Experimental/PythonScriptPlugin/Source/PythonScriptPlugin/Private/PyWrapperTypeRegistry.h
--- a/Engine/Plugins/Experimental/PythonScriptPlugin/Source/PythonScriptPlugin/Private/PyWrapperTypeRegistry.h
+++ b/Engine/Plugins/Experimental/PythonScriptPlugin/Source/PythonScriptPlugin/Private/PyWrapperTypeRegistry.h
@@ -10,6 +10,7 @@
 #include "PyConversionMethod.h"
 #include "PyOnlineDocsWriter.h"
 #include "Misc/EnumClassFlags.h"
+#include "UObject/GCObject.h"
 #include "UObject/WeakObjectPtr.h"
 #include "Templates/Function.h"
 
@@ -316,12 +317,16 @@ public:
 
 
 /** Singleton instance that handles re-instancing Python types */
-class FPyWrapperTypeReinstancer
+class FPyWrapperTypeReinstancer : public FGCObject
 {
 public:
 	/** Access the singleton instance */
 	static FPyWrapperTypeReinstancer& Get();
 
+	//~ FGCObject interface
+	virtual void AddReferencedObjects(FReferenceCollector& InCollector) override;
+	virtual FString GetReferencerName() const override;
+
 	/** Add a pending pair of classes to be re-instanced */
 	void AddPendingClass(UPythonGeneratedClass* OldClass, UPythonGeneratedClass* NewClass);
 
@@ -331,9 +336,6 @@ public:
 	/** Process any pending re-instance requests */
 	void ProcessPending();
 
-	/** Collect any referenced objects */
-	void AddReferencedObjects(FReferenceCollector& InCollector);
-
 private:
 	/** Pending pairs of classes that to be re-instanced */
 	TArray<TPair<TObjectPtr<UPythonGeneratedClass>, TObjectPtr<UPythonGeneratedClass>>> ClassesToReinstance;
@@ -343,7 +345,7 @@ private:
 };
 
 /** Singleton instance that maps Unreal types to Python types */
-class FPyWrapperTypeRegistry
+class FPyWrapperTypeRegistry : public FGCObject
 {
 public:
 	/** Struct used to build up a list of wrapped type references that still need to be generated */
@@ -365,6 +367,10 @@ public:
 	/** Access the singleton instance */
 	static FPyWrapperTypeRegistry& Get();
 
+	//~ FGCObject interface
+	virtual void AddReferencedObjects(FReferenceCollector& InCollector) override;
+	virtual FString GetReferencerName() const override;
+
 	/** Callback for when a Python module is dirtied */
 	FOnModuleDirtied& OnModuleDirtied()
 	{
