Remove the restriction that Unreal types generated from Python scripts must have globally unique names

There are two strategies used internally for achieving this:
1. If the Python script lives under the project or engine root (and ideally under a content mount point), then a package is created to host any types generated by that Python script.
2. If the Python script lives outside of any known portable path, then we take its absolute filename (normalized and lowercase) and append it to the generated type name as a hash.

These strategies allows a.py and b.py to both define an unreal.uthing() called "Foo" without those two generated types having a conflicting Unreal type name.

Note: For backwards compatibility, we leave a redirector from the old-style name to the most recent modular/hashed type name.

diff --git a/Engine/Plugins/Experimental/PythonScriptPlugin/Source/PythonScriptPlugin/Private/PyCore.cpp b/Engine/Plugins/Experimental/PythonScriptPlugin/Source/PythonScriptPlugin/Private/PyCore.cpp
--- a/Engine/Plugins/Experimental/PythonScriptPlugin/Source/PythonScriptPlugin/Private/PyCore.cpp
+++ b/Engine/Plugins/Experimental/PythonScriptPlugin/Source/PythonScriptPlugin/Private/PyCore.cpp
@@ -93,6 +93,32 @@ UObject* GetPythonTypeContainer()
 	return GetPythonTypeContainerSingleton().Get();
 }
 
+UObjectRedirector* CreatePythonTypeLegacyRedirector(const FString& ShortName, const FTopLevelAssetPath& GeneratedPathName)
+{
+	UObject* RedirectorOuter = GetPythonTypeContainer();
+	UObjectRedirector* Redirector = FindObject<UObjectRedirector>(RedirectorOuter, *ShortName);
+
+	// If the generated path name matches the redirector path name, then delete any existing redirector as we're about to create an object with that name
+	if (FTopLevelAssetPath(RedirectorOuter->GetFName(), *ShortName) == GeneratedPathName)
+	{
+		if (Redirector)
+		{
+			Redirector->DestinationObject = nullptr;
+			Redirector->ClearFlags(RF_Public | RF_Standalone);
+			Redirector->Rename(*FString::Printf(TEXT("%s_DELETED"), *ShortName), nullptr, REN_DontCreateRedirectors);
+		}
+		return nullptr;
+	}
+	
+	// Otherwise, create a redirector if one doesn't already exist
+	// The caller is responsible for setting DestinationObject
+	if (!Redirector)
+	{
+		Redirector = NewObject<UObjectRedirector>(RedirectorOuter, *ShortName, RF_Public | RF_Standalone | RF_Transient);
+	}
+	return Redirector;
+}
+
 
 FPyDelegateHandle* FPyDelegateHandle::CreateInstance(const FDelegateHandle& InValue)
 {
diff --git a/Engine/Plugins/Experimental/PythonScriptPlugin/Source/PythonScriptPlugin/Private/PyCore.h b/Engine/Plugins/Experimental/PythonScriptPlugin/Source/PythonScriptPlugin/Private/PyCore.h
--- a/Engine/Plugins/Experimental/PythonScriptPlugin/Source/PythonScriptPlugin/Private/PyCore.h
+++ b/Engine/Plugins/Experimental/PythonScriptPlugin/Source/PythonScriptPlugin/Private/PyCore.h
@@ -21,6 +21,16 @@ struct FSlowTask;
 /** Get the object that Python created types should be outered to */
 UObject* GetPythonTypeContainer();
 
+/**
+ * Create a redirector for a Python type, if needed.
+ * This is for backwards compatability, as all types were historically generated under GetPythonTypeContainer() using their short name.
+ * @note The caller must set DestinationObject to point at the generated type.
+ * 
+ * @param ShortName			The short name of the Python generated type.
+ * @param GeneratedPathName The path name that the Python generated type will actually use.
+ */
+UObjectRedirector* CreatePythonTypeLegacyRedirector(const FString& ShortName, const FTopLevelAssetPath& GeneratedPathName);
+
 /** Python type for FPyDelegateHandle */
 extern PyTypeObject PyDelegateHandleType;
 
diff --git a/Engine/Plugins/Experimental/PythonScriptPlugin/Source/PythonScriptPlugin/Private/PyReferenceCollector.cpp b/Engine/Plugins/Experimental/PythonScriptPlugin/Source/PythonScriptPlugin/Private/PyReferenceCollector.cpp
--- a/Engine/Plugins/Experimental/PythonScriptPlugin/Source/PythonScriptPlugin/Private/PyReferenceCollector.cpp
+++ b/Engine/Plugins/Experimental/PythonScriptPlugin/Source/PythonScriptPlugin/Private/PyReferenceCollector.cpp
@@ -8,6 +8,7 @@
 #include "PyWrapperEnum.h"
 #include "PyWrapperDelegate.h"
 #include "PyGIL.h"
+#include "UObject/ObjectRedirector.h"
 #include "UObject/UnrealType.h"
 #include "UObject/UObjectHash.h"
 #include "UObject/PurgingReferenceCollector.h"
@@ -83,7 +84,17 @@ void FPyReferenceCollector::PurgeUnrealGeneratedTypes()
 	FPurgingReferenceCollector PurgingReferenceCollector;
 	TArray<FWeakObjectPtr> WeakReferencesToPurgedObjects;
 
-	auto FlagObjectForPurge = [&PurgingReferenceCollector, &WeakReferencesToPurgedObjects](UObject* InObject, const bool bMarkPendingKill)
+	TMap<UObject*, UObjectRedirector*> ReverseRedirectorMapping;
+	ForEachObjectOfClass(UObjectRedirector::StaticClass(), [&ReverseRedirectorMapping](UObject* InObject)
+	{
+		UObjectRedirector* Redirector = CastChecked<UObjectRedirector>(InObject);
+		if (Redirector->DestinationObject)
+		{
+			ReverseRedirectorMapping.Add(Redirector->DestinationObject, Redirector);
+		}
+	});
+
+	auto FlagObjectForPurge = [&PurgingReferenceCollector, &WeakReferencesToPurgedObjects, &ReverseRedirectorMapping](UObject* InObject, const bool bMarkPendingKill)
 	{
 		check(!InObject->HasAnyInternalFlags(EInternalObjectFlags::Native));
 
@@ -102,6 +113,13 @@ void FPyReferenceCollector::PurgeUnrealGeneratedTypes()
 			WeakReferencesToPurgedObjects.Add(InObject);
 		}
 
+		if (UObjectRedirector* Redirector = ReverseRedirectorMapping.FindRef(InObject))
+		{
+			check(Redirector->DestinationObject == InObject);
+			Redirector->DestinationObject = nullptr;
+			Redirector->ClearFlags(RF_Public | RF_Standalone);
+		}
+
 		PurgingReferenceCollector.AddObjectToPurge(InObject);
 	};
 
diff --git a/Engine/Plugins/Experimental/PythonScriptPlugin/Source/PythonScriptPlugin/Private/PyUtil.cpp b/Engine/Plugins/Experimental/PythonScriptPlugin/Source/PythonScriptPlugin/Private/PyUtil.cpp
--- a/Engine/Plugins/Experimental/PythonScriptPlugin/Source/PythonScriptPlugin/Private/PyUtil.cpp
+++ b/Engine/Plugins/Experimental/PythonScriptPlugin/Source/PythonScriptPlugin/Private/PyUtil.cpp
@@ -18,6 +18,7 @@
 #include "PyWrapperFieldPath.h"
 #include "PyWrapperTypeRegistry.h"
 
+#include "Misc/PackageName.h"
 #include "Misc/Paths.h"
 #include "Misc/PathViews.h"
 #include "Misc/ScopeExit.h"
@@ -96,6 +97,29 @@ FString PyObjectToUEStringRepr(PyObject* InPyObj)
 	return PyObjectToUEString(InPyObj);
 }
 
+FEvalStack& FEvalStack::Get()
+{
+	static FEvalStack Instance;
+	return Instance;
+}
+
+void FEvalStack::PushContext(FEvalContext&& Context)
+{
+	Stack.Push(MoveTemp(Context));
+}
+
+void FEvalStack::PopContext()
+{
+	Stack.Pop();
+}
+
+const FEvalStack::FEvalContext* FEvalStack::GetCurrentContext() const
+{
+	return Stack.Num() > 0
+		? &Stack.Top()
+		: nullptr;
+}
+
 FPropValueOnScope::FPropValueOnScope(FConstPropOnScope&& InProp)
 	: Prop(MoveTemp(InProp))
 {
@@ -1241,6 +1265,95 @@ FString GetCleanTypename(PyObject* InPyObj)
 	return GetCleanTypename(PyType_Check(InPyObj) ? (PyTypeObject*)InPyObj : Py_TYPE(InPyObj));
 }
 
+void GetGeneratedTypeOuterAndName(PyTypeObject* InPyType, UObject*& OutOuter, FString& OutName)
+{
+	OutOuter = GetPythonTypeContainer();
+	OutName = GetCleanTypename(InPyType);
+
+	FString TypeFilename;
+
+	// Favor "inspect.getfile" if possible, as this will return the correct information for files executed via an import
+	if (FPyObjectPtr PyInspectModule = FPyObjectPtr::StealReference(PyImport_ImportModule("inspect")))
+	{
+		PyObject* PyInspectDict = PyModule_GetDict(PyInspectModule);
+		if (PyObject* PyGetFileFunc = PyDict_GetItemString(PyInspectDict, "getfile"))
+		{
+			if (FPyObjectPtr PyGetFileResult = FPyObjectPtr::StealReference(PyObject_CallFunctionObjArgs(PyGetFileFunc, InPyType, nullptr)))
+			{
+				TypeFilename = PyObjectToUEString(PyGetFileResult);
+			}
+			else
+			{
+				// Clear any exception information if getfile failed
+				PyErr_Clear();
+			}
+		}
+	}
+
+	// If "inspect.getfile" failed then try and access the current "__file__", as that will work for files being directly executed (ie, not an import)
+	if (TypeFilename.IsEmpty())
+	{
+		// Note: This doesn't use PyEval_GetGlobals() as that will return the result for the current "frame", which may be from an intermediate file (eg, unreal_core.py due to the unreal.uthing() decorator)
+		if (const FEvalStack::FEvalContext* CurrentContext = FEvalStack::Get().GetCurrentContext())
+		{
+			if (PyObject* PyGlobalsFile = PyDict_GetItemString(CurrentContext->GlobalDict, "__file__"))
+			{
+				TypeFilename = PyObjectToUEString(PyGlobalsFile);
+			}
+		}
+	}
+
+	// Normalize the found path for consistency (as the absolute path may vary)
+	if (!TypeFilename.IsEmpty())
+	{
+		FString TypePackageName;
+
+		FPaths::NormalizeFilename(TypeFilename);
+		if (!FPackageName::TryConvertFilenameToLongPackageName(TypeFilename, TypePackageName))
+		{
+			if (FPaths::IsUnderDirectory(TypeFilename, FPaths::EngineDir()))
+			{
+				FPaths::MakePathRelativeTo(TypeFilename, *FPaths::EngineDir());
+				TypePackageName = FPaths::Combine(TEXTVIEW("/Engine"), TypeFilename);
+			}
+			else if (FPaths::IsUnderDirectory(TypeFilename, FPaths::ProjectDir()))
+			{
+				FPaths::MakePathRelativeTo(TypeFilename, *FPaths::ProjectDir());
+				TypePackageName = FPaths::Combine(TEXTVIEW("/Game"), TypeFilename);
+			}
+		}
+			
+		if (TypePackageName.IsEmpty())
+		{
+			// This filename is something we can't resolve into a stable package path
+			// Just hash it into the type name to try and keep things unique
+			TypeFilename.ToLowerInline(); // To produce a case-insensitive hash
+			OutName += TStringBuilder<12>().Appendf(TEXT("_0x%08X"), FCrc::StrCrc32(*TypeFilename));
+		}
+		else
+		{
+			// Remove any remaining extension and add the "_PY" suffix
+			TypePackageName = FPaths::ChangeExtension(TypePackageName, TEXT(""));
+			TypePackageName += TEXTVIEW("_PY");
+
+			// This filename resolved into a stable package path, so put the generated type in that package
+			UPackage* TypePackage = FindObject<UPackage>(nullptr, *TypePackageName);
+			if (!TypePackage)
+			{
+				TypePackage = NewObject<UPackage>(nullptr, *TypePackageName, RF_Public | RF_Standalone | RF_Transient);
+				TypePackage->SetPackageFlags(PKG_ContainsScript);
+			}
+			OutOuter = TypePackage;
+		}
+	}
+}
+
+FString GetGeneratedTypeDisplayName(PyTypeObject* InPyType)
+{
+	FString CleanName = GetCleanTypename(InPyType);
+	return FName::NameToDisplayString(CleanName, /*bIsBool*/false);
+}
+
 FString GetErrorContext(PyTypeObject* InPyType)
 {
 	return UTF8_TO_TCHAR(InPyType->tp_name);
diff --git a/Engine/Plugins/Experimental/PythonScriptPlugin/Source/PythonScriptPlugin/Private/PyUtil.h b/Engine/Plugins/Experimental/PythonScriptPlugin/Source/PythonScriptPlugin/Private/PyUtil.h
--- a/Engine/Plugins/Experimental/PythonScriptPlugin/Source/PythonScriptPlugin/Private/PyUtil.h
+++ b/Engine/Plugins/Experimental/PythonScriptPlugin/Source/PythonScriptPlugin/Private/PyUtil.h
@@ -70,6 +70,27 @@ namespace PyUtil
 		return PyBool_FromLong(bResult);
 	}
 
+	/** Stack of context information given when calling PyEval_EvalCode */
+	class FEvalStack
+	{
+	public:
+		struct FEvalContext
+		{
+			const TCHAR* Context = nullptr;
+			PyObject* GlobalDict = nullptr;
+			PyObject* LocalDict = nullptr;
+		};
+
+		static FEvalStack& Get();
+
+		void PushContext(FEvalContext&& Context);
+		void PopContext();
+		const FEvalContext* GetCurrentContext() const;
+
+	private:
+		TArray<FEvalContext> Stack;
+	};
+
 	/** Helper to manage a property pointer that may be potentially owned by a wrapper or stack instance */
 	template <typename TPropType>
 	class TPropOnScope
@@ -568,6 +589,18 @@ namespace PyUtil
 		return GetCleanTypename((PyObject*)InPyObj);
 	}
 
+	/**
+	 * Get the generated type outer and name to use for the given Python type.
+	 * @note This should only be called when immediately processing the generation request 
+	 *       from Python, as it relies on getting information from the Python exec stack.
+	 */
+	void GetGeneratedTypeOuterAndName(PyTypeObject* InPyType, UObject*& OutOuter, FString& OutName);
+
+	/**
+	 * Get the generated display name to use for the given Python type.
+	 */
+	FString GetGeneratedTypeDisplayName(PyTypeObject* InPyType);
+
 	/**
 	 * Get the error context string of the given object.
 	 */
diff --git a/Engine/Plugins/Experimental/PythonScriptPlugin/Source/PythonScriptPlugin/Private/PyWrapperEnum.cpp b/Engine/Plugins/Experimental/PythonScriptPlugin/Source/PythonScriptPlugin/Private/PyWrapperEnum.cpp
--- a/Engine/Plugins/Experimental/PythonScriptPlugin/Source/PythonScriptPlugin/Private/PyWrapperEnum.cpp
+++ b/Engine/Plugins/Experimental/PythonScriptPlugin/Source/PythonScriptPlugin/Private/PyWrapperEnum.cpp
@@ -9,6 +9,7 @@
 
 #include "Containers/ArrayView.h"
 #include "Internationalization/Text.h"
+#include "UObject/ObjectRedirector.h"
 
 #include UE_INLINE_GENERATED_CPP_BY_NAME(PyWrapperEnum)
 
@@ -741,26 +742,30 @@ void FPyWrapperEnumMetaData::AddTypeReferencedObjects(FReferenceCollector& Colle
 class FPythonGeneratedEnumBuilder
 {
 public:
-	FPythonGeneratedEnumBuilder(const FString& InEnumName, PyTypeObject* InPyType)
-		: EnumName(InEnumName)
+	FPythonGeneratedEnumBuilder(PyTypeObject* InPyType)
+		: EnumName()
 		, PyType(InPyType)
 		, NewEnum(nullptr)
-		, bDidExist(false)
 	{
-		UObject* EnumOuter = GetPythonTypeContainer();
+		UObject* EnumOuter = nullptr;
+		PyUtil::GetGeneratedTypeOuterAndName(PyType, EnumOuter, EnumName);
+
+		UObjectRedirector* EnumRedirector = CreatePythonTypeLegacyRedirector(PyUtil::GetCleanTypename(InPyType), FTopLevelAssetPath(EnumOuter->GetFName(), *EnumName));
 
 		// Enum instances are re-used if they already exist
 		NewEnum = FindObject<UPythonGeneratedEnum>(EnumOuter, *EnumName);
-		if (NewEnum)
-		{
-			bDidExist = true;
-		}
-		else
+		if (!NewEnum)
 		{
 			NewEnum = NewObject<UPythonGeneratedEnum>(EnumOuter, *EnumName, RF_Public | RF_Standalone | RF_Transient);
+			NewEnum->SetMetaData(TEXT("DisplayName"), *PyUtil::GetGeneratedTypeDisplayName(PyType));
 			NewEnum->SetMetaData(TEXT("BlueprintType"), TEXT("true"));
 		}
 		NewEnum->EnumValueDefs.Reset();
+
+		if (EnumRedirector)
+		{
+			EnumRedirector->DestinationObject = NewEnum;
+		}
 	}
 
 	~FPythonGeneratedEnumBuilder()
@@ -798,7 +803,7 @@ public:
 
 		// Map the Unreal enum to the Python type
 		NewEnum->PyType = FPyTypeObjectPtr::NewReference(PyType);
-		FPyWrapperTypeRegistry::Get().RegisterWrappedEnumType(NewEnum, PyType, !bDidExist);
+		FPyWrapperTypeRegistry::Get().RegisterWrappedEnumType(NewEnum, PyType, false);
 
 		// Null the NewEnum pointer so the destructor doesn't kill it
 		UPythonGeneratedEnum* FinalizedEnum = NewEnum;
@@ -877,7 +882,6 @@ private:
 	FString EnumName;
 	PyTypeObject* PyType;
 	UPythonGeneratedEnum* NewEnum;
-	bool bDidExist;
 };
 
 void UPythonGeneratedEnum::BeginDestroy()
@@ -894,7 +898,7 @@ void UPythonGeneratedEnum::ReleasePythonResources()
 		FPyScopedGIL GIL;
 		if (PyType)
 		{
-			FPyWrapperTypeRegistry::Get().UnregisterWrappedEnumType(this, PyType, !HasAnyFlags(RF_NewerVersionExists));
+			FPyWrapperTypeRegistry::Get().UnregisterWrappedEnumType(this, PyType, false);
 		}
 		PyType.Reset();
 	}
@@ -911,7 +915,7 @@ void UPythonGeneratedEnum::ReleasePythonResources()
 UPythonGeneratedEnum* UPythonGeneratedEnum::GenerateEnum(PyTypeObject* InPyType)
 {
 	// Builder used to generate the enum
-	FPythonGeneratedEnumBuilder PythonEnumBuilder(PyUtil::GetCleanTypename(InPyType), InPyType);
+	FPythonGeneratedEnumBuilder PythonEnumBuilder(InPyType);
 
 	// Add the values to this enum
 	TArray<FPyUValueDef*> PyValueDefs;
diff --git a/Engine/Plugins/Experimental/PythonScriptPlugin/Source/PythonScriptPlugin/Private/PyWrapperEnum.h b/Engine/Plugins/Experimental/PythonScriptPlugin/Source/PythonScriptPlugin/Private/PyWrapperEnum.h
--- a/Engine/Plugins/Experimental/PythonScriptPlugin/Source/PythonScriptPlugin/Private/PyWrapperEnum.h
+++ b/Engine/Plugins/Experimental/PythonScriptPlugin/Source/PythonScriptPlugin/Private/PyWrapperEnum.h
@@ -113,7 +113,7 @@ typedef TPyPtr<FPyWrapperEnum> FPyWrapperEnumPtr;
 
 /** An Unreal enum that was generated from a Python type */
 UCLASS()
-class UPythonGeneratedEnum : public UEnum, public IPythonResourceOwner
+class UPythonGeneratedEnum final : public UEnum, public IPythonResourceOwner
 {
 	GENERATED_BODY()
 
@@ -122,6 +122,10 @@ class UPythonGeneratedEnum : public UEnum, public IPythonResourceOwner
 public:
 	//~ UObject interface
 	virtual void BeginDestroy() override;
+	virtual bool IsAsset() const override
+	{
+		return false;
+	}
 
 	//~ IPythonResourceOwner interface
 	virtual void ReleasePythonResources() override;
diff --git a/Engine/Plugins/Experimental/PythonScriptPlugin/Source/PythonScriptPlugin/Private/PyWrapperObject.cpp b/Engine/Plugins/Experimental/PythonScriptPlugin/Source/PythonScriptPlugin/Private/PyWrapperObject.cpp
--- a/Engine/Plugins/Experimental/PythonScriptPlugin/Source/PythonScriptPlugin/Private/PyWrapperObject.cpp
+++ b/Engine/Plugins/Experimental/PythonScriptPlugin/Source/PythonScriptPlugin/Private/PyWrapperObject.cpp
@@ -11,6 +11,7 @@
 #include "UObject/Package.h"
 #include "UObject/Class.h"
 #include "UObject/MetaData.h"
+#include "UObject/ObjectRedirector.h"
 #include "UObject/UnrealType.h"
 #include "UObject/UObjectHash.h"
 #include "UObject/StructOnScope.h"
@@ -1316,13 +1317,14 @@ bool FPyWrapperObjectMetaData::IsClassDeprecated(FPyWrapperObject* Instance, FSt
 class FPythonGeneratedClassBuilder
 {
 public:
-	FPythonGeneratedClassBuilder(const FString& InClassName, UClass* InSuperClass, PyTypeObject* InPyType)
-		: ClassName(InClassName)
+	FPythonGeneratedClassBuilder(UClass* InSuperClass, PyTypeObject* InPyType)
+		: ClassName()
 		, PyType(InPyType)
 		, OldClass(nullptr)
 		, NewClass(nullptr)
 	{
-		UObject* ClassOuter = GetPythonTypeContainer();
+		UObject* ClassOuter = nullptr;
+		PyUtil::GetGeneratedTypeOuterAndName(PyType, ClassOuter, ClassName);
 
 		// Find any existing class with the name we want to use
 		OldClass = FindObject<UPythonGeneratedClass>(ClassOuter, *ClassName);
@@ -1330,6 +1332,7 @@ public:
 		// Create a new class with a temporary name; we will rename it as part of Finalize
 		const FString NewClassName = MakeUniqueObjectName(ClassOuter, UPythonGeneratedClass::StaticClass(), *FString::Printf(TEXT("%s_NEWINST"), *ClassName)).ToString();
 		NewClass = NewObject<UPythonGeneratedClass>(ClassOuter, *NewClassName, RF_Public | RF_Standalone | RF_Transient);
+		NewClass->SetMetaData(TEXT("DisplayName"), *PyUtil::GetGeneratedTypeDisplayName(PyType));
 		NewClass->SetSuperStruct(InSuperClass);
 		NewClass->ClassFlags |= CLASS_Native;
 	}
@@ -1340,11 +1343,12 @@ public:
 		, OldClass(InOldClass)
 		, NewClass(nullptr)
 	{
-		UObject* ClassOuter = GetPythonTypeContainer();
+		UObject* ClassOuter = InOldClass->GetOuter();
 
 		// Create a new class with a temporary name; we will rename it as part of Finalize
 		const FString NewClassName = MakeUniqueObjectName(ClassOuter, UPythonGeneratedClass::StaticClass(), *FString::Printf(TEXT("%s_NEWINST"), *ClassName)).ToString();
 		NewClass = NewObject<UPythonGeneratedClass>(ClassOuter, *NewClassName, RF_Public | RF_Standalone | RF_Transient);
+		NewClass->SetMetaData(TEXT("DisplayName"), *PyUtil::GetGeneratedTypeDisplayName(PyType));
 		NewClass->SetSuperStruct(InSuperClass);
 		NewClass->ClassFlags |= CLASS_Native;
 	}
@@ -1388,6 +1392,10 @@ public:
 		}
 
 		Py_BEGIN_ALLOW_THREADS
+		if (UObjectRedirector* ClassRedirector = CreatePythonTypeLegacyRedirector(PyUtil::GetCleanTypename(PyType), FTopLevelAssetPath(NewClass->GetOuter()->GetFName(), *ClassName)))
+		{
+			ClassRedirector->DestinationObject = NewClass;
+		}
 		NewClass->Rename(*ClassName, nullptr, REN_DontCreateRedirectors);
 
 		// Finalize the class
@@ -1402,7 +1410,7 @@ public:
 
 		// Map the Unreal class to the Python type
 		NewClass->PyType = FPyTypeObjectPtr::NewReference(PyType);
-		FPyWrapperTypeRegistry::Get().RegisterWrappedClassType(NewClass, PyType);
+		FPyWrapperTypeRegistry::Get().RegisterWrappedClassType(NewClass, PyType, false);
 
 		// Ensure the CDO exists
 		Py_BEGIN_ALLOW_THREADS
@@ -1938,8 +1946,8 @@ void UPythonGeneratedClass::PostRename(UObject* OldOuter, const FName OldName)
 
 	if (PyType)
 	{
-		FPyWrapperTypeRegistry::Get().UnregisterWrappedClassType(FSoftObjectPath(OldOuter->GetFName(), OldName, FString()), PyType);
-		FPyWrapperTypeRegistry::Get().RegisterWrappedClassType(this, PyType, !HasAnyFlags(RF_NewerVersionExists));
+		FPyWrapperTypeRegistry::Get().UnregisterWrappedClassType(FSoftObjectPath(OldOuter->GetFName(), OldName, FString()), PyType, false);
+		FPyWrapperTypeRegistry::Get().RegisterWrappedClassType(this, PyType, false);
 	}
 }
 
@@ -1983,7 +1991,7 @@ void UPythonGeneratedClass::ReleasePythonResources()
 		FPyScopedGIL GIL;
 		if (PyType)
 		{
-			FPyWrapperTypeRegistry::Get().UnregisterWrappedClassType(this, PyType, !HasAnyFlags(RF_NewerVersionExists));
+			FPyWrapperTypeRegistry::Get().UnregisterWrappedClassType(this, PyType, false);
 		}
 		PyType.Reset();
 		PyPostInitFunction.Reset();
@@ -2025,7 +2033,7 @@ UPythonGeneratedClass* UPythonGeneratedClass::GenerateClass(PyTypeObject* InPyTy
 	}
 
 	// Builder used to generate the class
-	FPythonGeneratedClassBuilder PythonClassBuilder(PyUtil::GetCleanTypename(InPyType), SuperClass, InPyType);
+	FPythonGeneratedClassBuilder PythonClassBuilder(SuperClass, InPyType);
 
 	// Add the functions to this class
 	// We have to process these first as properties may reference them as get/set functions
diff --git a/Engine/Plugins/Experimental/PythonScriptPlugin/Source/PythonScriptPlugin/Private/PyWrapperObject.h b/Engine/Plugins/Experimental/PythonScriptPlugin/Source/PythonScriptPlugin/Private/PyWrapperObject.h
--- a/Engine/Plugins/Experimental/PythonScriptPlugin/Source/PythonScriptPlugin/Private/PyWrapperObject.h
+++ b/Engine/Plugins/Experimental/PythonScriptPlugin/Source/PythonScriptPlugin/Private/PyWrapperObject.h
@@ -184,7 +184,7 @@ typedef TPyPtr<FPyWrapperObject> FPyWrapperObjectPtr;
 
 /** An Unreal class that was generated from a Python type */
 UCLASS(BlueprintType, Transient)
-class UPythonGeneratedClass : public UClass, public IPythonResourceOwner
+class UPythonGeneratedClass final : public UClass, public IPythonResourceOwner
 {
 	GENERATED_BODY()
 
@@ -194,6 +194,10 @@ public:
 	//~ UObject interface
 	virtual void PostRename(UObject* OldOuter, const FName OldName) override;
 	virtual void BeginDestroy() override;
+	virtual bool IsAsset() const override
+	{
+		return false;
+	}
 
 	//~ UClass interface
 	virtual void PostInitInstance(UObject* InObj, FObjectInstancingGraph* InstanceGraph) override;
diff --git a/Engine/Plugins/Experimental/PythonScriptPlugin/Source/PythonScriptPlugin/Private/PyWrapperStruct.cpp b/Engine/Plugins/Experimental/PythonScriptPlugin/Source/PythonScriptPlugin/Private/PyWrapperStruct.cpp
--- a/Engine/Plugins/Experimental/PythonScriptPlugin/Source/PythonScriptPlugin/Private/PyWrapperStruct.cpp
+++ b/Engine/Plugins/Experimental/PythonScriptPlugin/Source/PythonScriptPlugin/Private/PyWrapperStruct.cpp
@@ -7,6 +7,7 @@
 #include "PyReferenceCollector.h"
 #include "UObject/Package.h"
 #include "UObject/Class.h"
+#include "UObject/ObjectRedirector.h"
 #include "UObject/PropertyPortFlags.h"
 #include "Misc/ScopeExit.h"
 #include "Templates/Casts.h"
@@ -1439,13 +1440,14 @@ bool FPyWrapperStructMetaData::IsStructDeprecated(FPyWrapperStruct* Instance, FS
 class FPythonGeneratedStructBuilder
 {
 public:
-	FPythonGeneratedStructBuilder(const FString& InStructName, UScriptStruct* InSuperStruct, PyTypeObject* InPyType)
-		: StructName(InStructName)
+	FPythonGeneratedStructBuilder(UScriptStruct* InSuperStruct, PyTypeObject* InPyType)
+		: StructName()
 		, PyType(InPyType)
 		, OldStruct(nullptr)
 		, NewStruct(nullptr)
 	{
-		UObject* StructOuter = GetPythonTypeContainer();
+		UObject* StructOuter = nullptr;
+		PyUtil::GetGeneratedTypeOuterAndName(PyType, StructOuter, StructName);
 
 		// Find any existing struct with the name we want to use
 		OldStruct = FindObject<UPythonGeneratedStruct>(StructOuter, *StructName);
@@ -1453,6 +1455,7 @@ public:
 		// Create a new struct with a temporary name; we will rename it as part of Finalize
 		const FString NewStructName = MakeUniqueObjectName(StructOuter, UPythonGeneratedStruct::StaticClass(), *FString::Printf(TEXT("%s_NEWINST"), *StructName)).ToString();
 		NewStruct = NewObject<UPythonGeneratedStruct>(StructOuter, *NewStructName, RF_Public | RF_Standalone | RF_Transient);
+		NewStruct->SetMetaData(TEXT("DisplayName"), *PyUtil::GetGeneratedTypeDisplayName(PyType));
 		NewStruct->SetMetaData(TEXT("BlueprintType"), TEXT("true"));
 		NewStruct->SetSuperStruct(InSuperStruct);
 	}
@@ -1511,6 +1514,11 @@ public:
 		{
 			PrepareOldStructForReinstancing();
 		}
+
+		if (UObjectRedirector* StructRedirector = CreatePythonTypeLegacyRedirector(PyUtil::GetCleanTypename(PyType), FTopLevelAssetPath(NewStruct->GetOuter()->GetFName(), *StructName)))
+		{
+			StructRedirector->DestinationObject = NewStruct;
+		}
 		NewStruct->Rename(*StructName, nullptr, REN_DontCreateRedirectors);
 
 		// Finalize the struct
@@ -1524,7 +1532,7 @@ public:
 
 		// Map the Unreal struct to the Python type
 		NewStruct->PyType = FPyTypeObjectPtr::NewReference(PyType);
-		FPyWrapperTypeRegistry::Get().RegisterWrappedStructType(NewStruct, PyType);
+		FPyWrapperTypeRegistry::Get().RegisterWrappedStructType(NewStruct, PyType, false);
 
 		// Re-instance the old struct
 		if (OldStruct)
@@ -1623,8 +1631,8 @@ void UPythonGeneratedStruct::PostRename(UObject* OldOuter, const FName OldName)
 
 	if (PyType)
 	{
-		FPyWrapperTypeRegistry::Get().UnregisterWrappedStructType(FSoftObjectPath(OldOuter->GetFName(), OldName, FString()), PyType);
-		FPyWrapperTypeRegistry::Get().RegisterWrappedStructType(this, PyType, !HasAnyFlags(RF_NewerVersionExists));
+		FPyWrapperTypeRegistry::Get().UnregisterWrappedStructType(FSoftObjectPath(OldOuter->GetFName(), OldName, FString()), PyType, false);
+		FPyWrapperTypeRegistry::Get().RegisterWrappedStructType(this, PyType, false);
 	}
 }
 
@@ -1673,7 +1681,7 @@ void UPythonGeneratedStruct::ReleasePythonResources()
 		FPyScopedGIL GIL;
 		if (PyType)
 		{
-			FPyWrapperTypeRegistry::Get().UnregisterWrappedStructType(this, PyType, !HasAnyFlags(RF_NewerVersionExists));
+			FPyWrapperTypeRegistry::Get().UnregisterWrappedStructType(this, PyType, false);
 		}
 		PyType.Reset();
 		PyPostInitFunction.Reset();
@@ -1704,7 +1712,7 @@ UPythonGeneratedStruct* UPythonGeneratedStruct::GenerateStruct(PyTypeObject* InP
 	}
 
 	// Builder used to generate the struct
-	FPythonGeneratedStructBuilder PythonStructBuilder(PyUtil::GetCleanTypename(InPyType), SuperStruct, InPyType);
+	FPythonGeneratedStructBuilder PythonStructBuilder(SuperStruct, InPyType);
 
 	// Add the fields to this struct
 	{
diff --git a/Engine/Plugins/Experimental/PythonScriptPlugin/Source/PythonScriptPlugin/Private/PyWrapperStruct.h b/Engine/Plugins/Experimental/PythonScriptPlugin/Source/PythonScriptPlugin/Private/PyWrapperStruct.h
--- a/Engine/Plugins/Experimental/PythonScriptPlugin/Source/PythonScriptPlugin/Private/PyWrapperStruct.h
+++ b/Engine/Plugins/Experimental/PythonScriptPlugin/Source/PythonScriptPlugin/Private/PyWrapperStruct.h
@@ -372,7 +372,7 @@ typedef TPyPtr<FPyWrapperStruct> FPyWrapperStructPtr;
 
 /** An Unreal struct that was generated from a Python type */
 UCLASS()
-class UPythonGeneratedStruct : public UScriptStruct, public IPythonResourceOwner
+class UPythonGeneratedStruct final : public UScriptStruct, public IPythonResourceOwner
 {
 	GENERATED_BODY()
 
@@ -382,6 +382,10 @@ public:
 	//~ UObject interface
 	virtual void PostRename(UObject* OldOuter, const FName OldName) override;
 	virtual void BeginDestroy() override;
+	virtual bool IsAsset() const override
+	{
+		return false;
+	}
 
 	//~ UStruct interface
 	virtual void InitializeStruct(void* Dest, int32 ArrayDim = 1) const override;
diff --git a/Engine/Plugins/Experimental/PythonScriptPlugin/Source/PythonScriptPlugin/Private/PythonScriptPlugin.cpp b/Engine/Plugins/Experimental/PythonScriptPlugin/Source/PythonScriptPlugin/Private/PythonScriptPlugin.cpp
--- a/Engine/Plugins/Experimental/PythonScriptPlugin/Source/PythonScriptPlugin/Private/PythonScriptPlugin.cpp
+++ b/Engine/Plugins/Experimental/PythonScriptPlugin/Source/PythonScriptPlugin/Private/PythonScriptPlugin.cpp
@@ -1428,7 +1428,11 @@ PyObject* FPythonScriptPlugin::EvalString(const TCHAR* InStr, const TCHAR* InCon
 		return nullptr;
 	}
 
-	return PyEval_EvalCode((PyUtil::FPyCodeObjectType*)PyCodeObj.Get(), InGlobalDict, InLocalDict);
+	PyUtil::FEvalStack::Get().PushContext(PyUtil::FEvalStack::FEvalContext{ InContext, InGlobalDict, InLocalDict});
+	PyObject* PyEvalResult = PyEval_EvalCode((PyUtil::FPyCodeObjectType*)PyCodeObj.Get(), InGlobalDict, InLocalDict);
+	PyUtil::FEvalStack::Get().PopContext();
+
+	return PyEvalResult;
 }
 
 bool FPythonScriptPlugin::RunString(FPythonCommandEx& InOutPythonCommand)
