Remove RF_Standalone from Python generated types to ensure consistent GC behavior between a full editor, a commandlet, and -game/-server

This will make it easier to find Python GC issues generally, rather than them only appearing in certain edge-case situations (like running tests in -game).

This also ensures that generated types are GC'd (during shutdown or reinstancing) by unregistering them from FPyWrapperTypeRegistry, as otherwise it would keep them alive as an ARO reference.

#rnx

diff --git a/Engine/Plugins/Experimental/PythonScriptPlugin/Source/PythonScriptPlugin/Private/PyCore.cpp b/Engine/Plugins/Experimental/PythonScriptPlugin/Source/PythonScriptPlugin/Private/PyCore.cpp
--- a/Engine/Plugins/Experimental/PythonScriptPlugin/Source/PythonScriptPlugin/Private/PyCore.cpp
+++ b/Engine/Plugins/Experimental/PythonScriptPlugin/Source/PythonScriptPlugin/Private/PyCore.cpp
@@ -114,7 +114,7 @@ UObjectRedirector* CreatePythonTypeLegacyRedirector(const FString& ShortName, co
 	// The caller is responsible for setting DestinationObject
 	if (!Redirector)
 	{
-		Redirector = NewObject<UObjectRedirector>(RedirectorOuter, *ShortName, RF_Public | RF_Standalone | RF_Transient);
+		Redirector = NewObject<UObjectRedirector>(RedirectorOuter, *ShortName, RF_Public | RF_Transient);
 	}
 	return Redirector;
 }
diff --git a/Engine/Plugins/Experimental/PythonScriptPlugin/Source/PythonScriptPlugin/Private/PyReferenceCollector.cpp b/Engine/Plugins/Experimental/PythonScriptPlugin/Source/PythonScriptPlugin/Private/PyReferenceCollector.cpp
--- a/Engine/Plugins/Experimental/PythonScriptPlugin/Source/PythonScriptPlugin/Private/PyReferenceCollector.cpp
+++ b/Engine/Plugins/Experimental/PythonScriptPlugin/Source/PythonScriptPlugin/Private/PyReferenceCollector.cpp
@@ -128,9 +128,12 @@ void FPyReferenceCollector::PurgeUnrealGeneratedTypes()
 	{
 		ForEachObjectOfClass(UPythonGeneratedClass::StaticClass(), [&FlagObjectForPurge](UObject* InObject)
 		{
-			FlagObjectForPurge(InObject, /*bMarkPendingKill*/false);
-
 			UPythonGeneratedClass* PythonGeneratedClass = CastChecked<UPythonGeneratedClass>(InObject);
+
+			// Generated classes are kept alive by the FPyWrapperTypeRegistry ARO, so unregister the type from FPyWrapperTypeRegistry
+			PythonGeneratedClass->UnregisterGeneratedType();
+
+			FlagObjectForPurge(InObject, /*bMarkPendingKill*/false);
 			ForEachObjectOfClass(PythonGeneratedClass, [&FlagObjectForPurge](UObject* InInnerObject)
 			{
 				FlagObjectForPurge(InInnerObject, /*bMarkPendingKill*/true);
@@ -143,6 +146,11 @@ void FPyReferenceCollector::PurgeUnrealGeneratedTypes()
 	{
 		ForEachObjectOfClass(UPythonGeneratedStruct::StaticClass(), [&FlagObjectForPurge](UObject* InObject)
 		{
+			UPythonGeneratedStruct* PythonGeneratedStruct = CastChecked<UPythonGeneratedStruct>(InObject);
+
+			// Generated structs are kept alive by the FPyWrapperTypeRegistry ARO, so unregister the type from FPyWrapperTypeRegistry
+			PythonGeneratedStruct->UnregisterGeneratedType();
+
 			FlagObjectForPurge(InObject, /*bMarkPendingKill*/false);
 		}, false, RF_ClassDefaultObject, EInternalObjectFlags::Native);
 	}
@@ -152,6 +160,11 @@ void FPyReferenceCollector::PurgeUnrealGeneratedTypes()
 	{
 		ForEachObjectOfClass(UPythonGeneratedEnum::StaticClass(), [&FlagObjectForPurge](UObject* InObject)
 		{
+			UPythonGeneratedEnum* PythonGeneratedEnum = CastChecked<UPythonGeneratedEnum>(InObject);
+
+			// Generated enums are kept alive by the FPyWrapperTypeRegistry ARO, so unregister the type from FPyWrapperTypeRegistry
+			PythonGeneratedEnum->UnregisterGeneratedType();
+
 			FlagObjectForPurge(InObject, /*bMarkPendingKill*/false);
 		}, false, RF_ClassDefaultObject, EInternalObjectFlags::Native);
 	}
diff --git a/Engine/Plugins/Experimental/PythonScriptPlugin/Source/PythonScriptPlugin/Private/PyUtil.cpp b/Engine/Plugins/Experimental/PythonScriptPlugin/Source/PythonScriptPlugin/Private/PyUtil.cpp
--- a/Engine/Plugins/Experimental/PythonScriptPlugin/Source/PythonScriptPlugin/Private/PyUtil.cpp
+++ b/Engine/Plugins/Experimental/PythonScriptPlugin/Source/PythonScriptPlugin/Private/PyUtil.cpp
@@ -1340,7 +1340,7 @@ void GetGeneratedTypeOuterAndName(PyTypeObject* InPyType, UObject*& OutOuter, FS
 			UPackage* TypePackage = FindObject<UPackage>(nullptr, *TypePackageName);
 			if (!TypePackage)
 			{
-				TypePackage = NewObject<UPackage>(nullptr, *TypePackageName, RF_Public | RF_Standalone | RF_Transient);
+				TypePackage = NewObject<UPackage>(nullptr, *TypePackageName, RF_Public | RF_Transient);
 				TypePackage->SetPackageFlags(PKG_ContainsScript);
 			}
 			OutOuter = TypePackage;
diff --git a/Engine/Plugins/Experimental/PythonScriptPlugin/Source/PythonScriptPlugin/Private/PyWrapperEnum.cpp b/Engine/Plugins/Experimental/PythonScriptPlugin/Source/PythonScriptPlugin/Private/PyWrapperEnum.cpp
--- a/Engine/Plugins/Experimental/PythonScriptPlugin/Source/PythonScriptPlugin/Private/PyWrapperEnum.cpp
+++ b/Engine/Plugins/Experimental/PythonScriptPlugin/Source/PythonScriptPlugin/Private/PyWrapperEnum.cpp
@@ -756,7 +756,7 @@ public:
 		NewEnum = FindObject<UPythonGeneratedEnum>(EnumOuter, *EnumName);
 		if (!NewEnum)
 		{
-			NewEnum = NewObject<UPythonGeneratedEnum>(EnumOuter, *EnumName, RF_Public | RF_Standalone | RF_Transient);
+			NewEnum = NewObject<UPythonGeneratedEnum>(EnumOuter, *EnumName, RF_Public | RF_Transient);
 			NewEnum->SetMetaData(TEXT("DisplayName"), *PyUtil::GetGeneratedTypeDisplayName(PyType));
 			NewEnum->SetMetaData(TEXT("BlueprintType"), TEXT("true"));
 		}
@@ -896,10 +896,7 @@ void UPythonGeneratedEnum::ReleasePythonResources()
 	if (Py_IsInitialized())
 	{
 		FPyScopedGIL GIL;
-		if (PyType)
-		{
-			FPyWrapperTypeRegistry::Get().UnregisterWrappedEnumType(this, PyType, false);
-		}
+		UnregisterGeneratedType();
 		PyType.Reset();
 	}
 	else
@@ -912,6 +909,14 @@ void UPythonGeneratedEnum::ReleasePythonResources()
 	PyMetaData = FPyWrapperEnumMetaData();
 }
 
+void UPythonGeneratedEnum::UnregisterGeneratedType()
+{
+	if (PyType)
+	{
+		FPyWrapperTypeRegistry::Get().UnregisterWrappedEnumType(this, PyType, false);
+	}
+}
+
 UPythonGeneratedEnum* UPythonGeneratedEnum::GenerateEnum(PyTypeObject* InPyType)
 {
 	// Builder used to generate the enum
diff --git a/Engine/Plugins/Experimental/PythonScriptPlugin/Source/PythonScriptPlugin/Private/PyWrapperEnum.h b/Engine/Plugins/Experimental/PythonScriptPlugin/Source/PythonScriptPlugin/Private/PyWrapperEnum.h
--- a/Engine/Plugins/Experimental/PythonScriptPlugin/Source/PythonScriptPlugin/Private/PyWrapperEnum.h
+++ b/Engine/Plugins/Experimental/PythonScriptPlugin/Source/PythonScriptPlugin/Private/PyWrapperEnum.h
@@ -130,6 +130,9 @@ public:
 	//~ IPythonResourceOwner interface
 	virtual void ReleasePythonResources() override;
 
+	/** Unregister this type from FPyWrapperTypeRegistry */
+	void UnregisterGeneratedType();
+
 	/** Generate an Unreal enum from the given Python type */
 	static UPythonGeneratedEnum* GenerateEnum(PyTypeObject* InPyType);
 
diff --git a/Engine/Plugins/Experimental/PythonScriptPlugin/Source/PythonScriptPlugin/Private/PyWrapperObject.cpp b/Engine/Plugins/Experimental/PythonScriptPlugin/Source/PythonScriptPlugin/Private/PyWrapperObject.cpp
--- a/Engine/Plugins/Experimental/PythonScriptPlugin/Source/PythonScriptPlugin/Private/PyWrapperObject.cpp
+++ b/Engine/Plugins/Experimental/PythonScriptPlugin/Source/PythonScriptPlugin/Private/PyWrapperObject.cpp
@@ -1331,7 +1331,7 @@ public:
 
 		// Create a new class with a temporary name; we will rename it as part of Finalize
 		const FString NewClassName = MakeUniqueObjectName(ClassOuter, UPythonGeneratedClass::StaticClass(), *FString::Printf(TEXT("%s_NEWINST"), *ClassName)).ToString();
-		NewClass = NewObject<UPythonGeneratedClass>(ClassOuter, *NewClassName, RF_Public | RF_Standalone | RF_Transient);
+		NewClass = NewObject<UPythonGeneratedClass>(ClassOuter, *NewClassName, RF_Public | RF_Transient);
 		NewClass->SetMetaData(TEXT("DisplayName"), *PyUtil::GetGeneratedTypeDisplayName(PyType));
 		NewClass->SetSuperStruct(InSuperClass);
 		NewClass->ClassFlags |= CLASS_Native;
@@ -1347,7 +1347,7 @@ public:
 
 		// Create a new class with a temporary name; we will rename it as part of Finalize
 		const FString NewClassName = MakeUniqueObjectName(ClassOuter, UPythonGeneratedClass::StaticClass(), *FString::Printf(TEXT("%s_NEWINST"), *ClassName)).ToString();
-		NewClass = NewObject<UPythonGeneratedClass>(ClassOuter, *NewClassName, RF_Public | RF_Standalone | RF_Transient);
+		NewClass = NewObject<UPythonGeneratedClass>(ClassOuter, *NewClassName, RF_Public | RF_Transient);
 		NewClass->SetMetaData(TEXT("DisplayName"), *PyUtil::GetGeneratedTypeDisplayName(PyType));
 		NewClass->SetSuperStruct(InSuperClass);
 		NewClass->ClassFlags |= CLASS_Native;
@@ -1932,6 +1932,7 @@ private:
 		OldClass->SetFlags(RF_NewerVersionExists);
 		OldClass->ClearFlags(RF_Public | RF_Standalone);
 		OldClass->Rename(*OldClassName, nullptr, REN_DontCreateRedirectors);
+		OldClass->UnregisterGeneratedType();
 	}
 
 	FString ClassName;
@@ -1989,10 +1990,7 @@ void UPythonGeneratedClass::ReleasePythonResources()
 	if (Py_IsInitialized())
 	{
 		FPyScopedGIL GIL;
-		if (PyType)
-		{
-			FPyWrapperTypeRegistry::Get().UnregisterWrappedClassType(this, PyType, false);
-		}
+		UnregisterGeneratedType();
 		PyType.Reset();
 		PyPostInitFunction.Reset();
 		for (const TSharedPtr<PyGenUtil::FFunctionDef>& FunctionDef : FunctionDefs)
@@ -2016,6 +2014,14 @@ void UPythonGeneratedClass::ReleasePythonResources()
 	PyMetaData = FPyWrapperObjectMetaData();
 }
 
+void UPythonGeneratedClass::UnregisterGeneratedType()
+{
+	if (PyType)
+	{
+		FPyWrapperTypeRegistry::Get().UnregisterWrappedClassType(this, PyType, false);
+	}
+}
+
 bool UPythonGeneratedClass::IsFunctionImplementedInScript(FName InFunctionName) const
 {
 	UFunction* Function = FindFunctionByName(InFunctionName);
diff --git a/Engine/Plugins/Experimental/PythonScriptPlugin/Source/PythonScriptPlugin/Private/PyWrapperObject.h b/Engine/Plugins/Experimental/PythonScriptPlugin/Source/PythonScriptPlugin/Private/PyWrapperObject.h
--- a/Engine/Plugins/Experimental/PythonScriptPlugin/Source/PythonScriptPlugin/Private/PyWrapperObject.h
+++ b/Engine/Plugins/Experimental/PythonScriptPlugin/Source/PythonScriptPlugin/Private/PyWrapperObject.h
@@ -205,6 +205,9 @@ public:
 	//~ IPythonResourceOwner interface
 	virtual void ReleasePythonResources() override;
 
+	/** Unregister this type from FPyWrapperTypeRegistry */
+	void UnregisterGeneratedType();
+
 	virtual bool IsFunctionImplementedInScript(FName InFunctionName) const override;
 
 	/** Generate an Unreal class from the given Python type */
diff --git a/Engine/Plugins/Experimental/PythonScriptPlugin/Source/PythonScriptPlugin/Private/PyWrapperStruct.cpp b/Engine/Plugins/Experimental/PythonScriptPlugin/Source/PythonScriptPlugin/Private/PyWrapperStruct.cpp
--- a/Engine/Plugins/Experimental/PythonScriptPlugin/Source/PythonScriptPlugin/Private/PyWrapperStruct.cpp
+++ b/Engine/Plugins/Experimental/PythonScriptPlugin/Source/PythonScriptPlugin/Private/PyWrapperStruct.cpp
@@ -1454,7 +1454,7 @@ public:
 
 		// Create a new struct with a temporary name; we will rename it as part of Finalize
 		const FString NewStructName = MakeUniqueObjectName(StructOuter, UPythonGeneratedStruct::StaticClass(), *FString::Printf(TEXT("%s_NEWINST"), *StructName)).ToString();
-		NewStruct = NewObject<UPythonGeneratedStruct>(StructOuter, *NewStructName, RF_Public | RF_Standalone | RF_Transient);
+		NewStruct = NewObject<UPythonGeneratedStruct>(StructOuter, *NewStructName, RF_Public | RF_Transient);
 		NewStruct->SetMetaData(TEXT("DisplayName"), *PyUtil::GetGeneratedTypeDisplayName(PyType));
 		NewStruct->SetMetaData(TEXT("BlueprintType"), TEXT("true"));
 		NewStruct->SetSuperStruct(InSuperStruct);
@@ -1617,6 +1617,7 @@ private:
 		OldStruct->SetFlags(RF_NewerVersionExists);
 		OldStruct->ClearFlags(RF_Public | RF_Standalone);
 		OldStruct->Rename(*OldStructName, nullptr, REN_DontCreateRedirectors);
+		OldStruct->UnregisterGeneratedType();
 	}
 
 	FString StructName;
@@ -1679,10 +1680,7 @@ void UPythonGeneratedStruct::ReleasePythonResources()
 	if (Py_IsInitialized())
 	{
 		FPyScopedGIL GIL;
-		if (PyType)
-		{
-			FPyWrapperTypeRegistry::Get().UnregisterWrappedStructType(this, PyType, false);
-		}
+		UnregisterGeneratedType();
 		PyType.Reset();
 		PyPostInitFunction.Reset();
 	}
@@ -1697,6 +1695,14 @@ void UPythonGeneratedStruct::ReleasePythonResources()
 	PyMetaData = FPyWrapperStructMetaData();
 }
 
+void UPythonGeneratedStruct::UnregisterGeneratedType()
+{
+	if (PyType)
+	{
+		FPyWrapperTypeRegistry::Get().UnregisterWrappedStructType(this, PyType, false);
+	}
+}
+
 UPythonGeneratedStruct* UPythonGeneratedStruct::GenerateStruct(PyTypeObject* InPyType)
 {
 	// Get the correct super struct from the parent type in Python
diff --git a/Engine/Plugins/Experimental/PythonScriptPlugin/Source/PythonScriptPlugin/Private/PyWrapperStruct.h b/Engine/Plugins/Experimental/PythonScriptPlugin/Source/PythonScriptPlugin/Private/PyWrapperStruct.h
--- a/Engine/Plugins/Experimental/PythonScriptPlugin/Source/PythonScriptPlugin/Private/PyWrapperStruct.h
+++ b/Engine/Plugins/Experimental/PythonScriptPlugin/Source/PythonScriptPlugin/Private/PyWrapperStruct.h
@@ -393,6 +393,9 @@ public:
 	//~ IPythonResourceOwner interface
 	virtual void ReleasePythonResources() override;
 
+	/** Unregister this type from FPyWrapperTypeRegistry */
+	void UnregisterGeneratedType();
+
 	/** Generate an Unreal struct from the given Python type */
 	static UPythonGeneratedStruct* GenerateStruct(PyTypeObject* InPyType);
 
diff --git a/Engine/Plugins/Experimental/PythonScriptPlugin/Source/PythonScriptPlugin/Private/PyWrapperTypeRegistry.cpp b/Engine/Plugins/Experimental/PythonScriptPlugin/Source/PythonScriptPlugin/Private/PyWrapperTypeRegistry.cpp
--- a/Engine/Plugins/Experimental/PythonScriptPlugin/Source/PythonScriptPlugin/Private/PyWrapperTypeRegistry.cpp
+++ b/Engine/Plugins/Experimental/PythonScriptPlugin/Source/PythonScriptPlugin/Private/PyWrapperTypeRegistry.cpp
@@ -2483,7 +2483,7 @@ PyTypeObject* FPyWrapperTypeRegistry::GenerateWrappedDelegateType(const UFunctio
 	// Generate the proxy class needed to wrap Python callables in Unreal delegates
 	UClass* PythonCallableForDelegateClass = nullptr;
 	{
-		PythonCallableForDelegateClass = NewObject<UClass>(GetPythonTypeContainer(), *PythonCallableForDelegateObjectName, RF_Public | RF_Standalone | RF_Transient);
+		PythonCallableForDelegateClass = NewObject<UClass>(GetPythonTypeContainer(), *PythonCallableForDelegateObjectName, RF_Public | RF_Transient);
 		UFunction* PythonCallableForDelegateFunc = nullptr;
 		{
 			FObjectDuplicationParameters FuncDuplicationParams(const_cast<UFunction*>(InDelegateSignature), PythonCallableForDelegateClass);
