[PythonScriptPlugin] Add script name to Insights tag when calling RunFile.

diff --git a/Engine/Plugins/Experimental/PythonScriptPlugin/Source/PythonScriptPlugin/Private/PythonScriptPlugin.cpp b/Engine/Plugins/Experimental/PythonScriptPlugin/Source/PythonScriptPlugin/Private/PythonScriptPlugin.cpp
--- a/Engine/Plugins/Experimental/PythonScriptPlugin/Source/PythonScriptPlugin/Private/PythonScriptPlugin.cpp
+++ b/Engine/Plugins/Experimental/PythonScriptPlugin/Source/PythonScriptPlugin/Private/PythonScriptPlugin.cpp
@@ -1481,7 +1481,7 @@ bool FPythonScriptPlugin::RunString(FPythonCommandEx& InOutPythonCommand)
 
 bool FPythonScriptPlugin::RunFile(const TCHAR* InFile, const TCHAR* InArgs, FPythonCommandEx& InOutPythonCommand)
 {
-	TRACE_CPUPROFILER_EVENT_SCOPE(FPythonScriptPlugin::RunFile)
+	TRACE_CPUPROFILER_EVENT_SCOPE_TEXT(*FString::Printf(TEXT("FPythonScriptPlugin::RunFile(%s)"), InFile ? InFile : TEXT("null")));
 
 	auto ResolveFilePath = [InFile]() -> FString
 	{
