[VFX Reference Update] Make Python EvalString more tolerant of encoding and locale side effects from executing code.
Improved log outputs for running python scripts to fall under LogPython.
Revert changes that temporarily moved python initialization of Neural Morph Model plugin to avoid encoding side effects from imports.

diff --git a/Engine/Plugins/Experimental/PythonScriptPlugin/Source/PythonScriptPlugin/Private/PythonScriptPlugin.cpp b/Engine/Plugins/Experimental/PythonScriptPlugin/Source/PythonScriptPlugin/Private/PythonScriptPlugin.cpp
--- a/Engine/Plugins/Experimental/PythonScriptPlugin/Source/PythonScriptPlugin/Private/PythonScriptPlugin.cpp
+++ b/Engine/Plugins/Experimental/PythonScriptPlugin/Source/PythonScriptPlugin/Private/PythonScriptPlugin.cpp
@@ -140,6 +140,80 @@
 	TArray<PyUtil::FPyApiChar*> PyCommandLineArgPtrs;
 };
 
+struct FScopedEncodingGuard
+{
+	FScopedEncodingGuard()
+	{
+		// Python 3 changes the console mode from O_TEXT to O_BINARY which affects other uses of the console
+		// So change the console mode back to its current setting after Py_Initialize has been called
+#if PLATFORM_WINDOWS
+		// We call _setmode here to cache the current state
+		CA_SUPPRESS(6031)
+		fflush(stdin);
+		StdInMode = _setmode(_fileno(stdin), _O_TEXT);
+		CA_SUPPRESS(6031)
+		fflush(stdout);
+		StdOutMode = _setmode(_fileno(stdout), _O_TEXT);
+		CA_SUPPRESS(6031)
+		fflush(stderr);
+		StdErrMode = _setmode(_fileno(stderr), _O_TEXT);
+#endif	// PLATFORM_WINDOWS
+
+#if PY_MAJOR_VERSION >= 3 && PY_MINOR_VERSION >= 7
+		// Python 3.7+ changes the C locale which affects functions using C string APIs
+		// So change the C locale back to its current setting after Py_Initialize has been called
+		if (const char* CurrentLocalePtr = setlocale(LC_ALL, nullptr))
+		{
+			CurrentLocale = ANSI_TO_TCHAR(CurrentLocalePtr);
+		}
+#endif	// PY_MAJOR_VERSION >= 3 && PY_MINOR_VERSION >= 7
+	}
+
+	~FScopedEncodingGuard()
+	{
+#if PLATFORM_WINDOWS
+		// We call _setmode here to restore the previous state
+		if (StdInMode != -1)
+		{
+			CA_SUPPRESS(6031)
+				fflush(stdin);
+			CA_SUPPRESS(6031)
+				_setmode(_fileno(stdin), StdInMode);
+		}
+		if (StdOutMode != -1)
+		{
+			CA_SUPPRESS(6031)
+				fflush(stdout);
+			CA_SUPPRESS(6031)
+				_setmode(_fileno(stdout), StdOutMode);
+		}
+		if (StdErrMode != -1)
+		{
+			CA_SUPPRESS(6031)
+				fflush(stderr);
+			CA_SUPPRESS(6031)
+				_setmode(_fileno(stderr), StdErrMode);
+		}
+#endif	// PLATFORM_WINDOWS
+
+#if PY_MAJOR_VERSION >= 3 && PY_MINOR_VERSION >= 7
+		// We call setlocale here to restore the previous state
+		if (!CurrentLocale.IsEmpty())
+		{
+			setlocale(LC_ALL, TCHAR_TO_ANSI(*CurrentLocale));
+		}
+#endif	// PY_MAJOR_VERSION >= 3 && PY_MINOR_VERSION >= 7
+	}
+
+private:
+	FString CurrentLocale;
+#if PLATFORM_WINDOWS
+	int StdErrMode = -1;
+	int StdOutMode = -1;
+	int StdInMode = -1;
+#endif	// PLATFORM_WINDOWS
+};
+
 FPythonCommandExecutor::FPythonCommandExecutor(IPythonScriptPlugin* InPythonScriptPlugin)
 	: PythonScriptPlugin(InPythonScriptPlugin)
 {
@@ -850,30 +924,7 @@
 		static_assert(PY_MAJOR_VERSION >= 3, "Unreal Engine Python integration doesn't support versions prior to Python 3.x");
 		UE_LOG(LogPython, Log, TEXT("Using Python %d.%d.%d"), PY_MAJOR_VERSION, PY_MINOR_VERSION, PY_MICRO_VERSION);
 
-		// Python 3 changes the console mode from O_TEXT to O_BINARY which affects other uses of the console
-		// So change the console mode back to its current setting after Py_Initialize has been called
-#if PLATFORM_WINDOWS
-		// We call _setmode here to cache the current state
-		CA_SUPPRESS(6031)
-		fflush(stdin);
-		const int StdInMode  = _setmode(_fileno(stdin), _O_TEXT);
-		CA_SUPPRESS(6031)
-		fflush(stdout);
-		const int StdOutMode = _setmode(_fileno(stdout), _O_TEXT);
-		CA_SUPPRESS(6031)
-		fflush(stderr);
-		const int StdErrMode = _setmode(_fileno(stderr), _O_TEXT);
-#endif	// PLATFORM_WINDOWS
+		FScopedEncodingGuard EncodingGuard;
 
-#if PY_MAJOR_VERSION >= 3 && PY_MINOR_VERSION >= 7
-		// Python 3.7+ changes the C locale which affects functions using C string APIs
-		// So change the C locale back to its current setting after Py_Initialize has been called
-		FString CurrentLocale;
-		if (const char* CurrentLocalePtr = setlocale(LC_ALL, nullptr))
-		{
-			CurrentLocale = ANSI_TO_TCHAR(CurrentLocalePtr);
-		}
-#endif	// PY_MAJOR_VERSION >= 3 && PY_MINOR_VERSION >= 7
-
 		// Check if the interpreter is should run in isolation mode.
 		int IsolatedInterpreterFlag = PythonPluginSettings->bIsolateInterpreterEnvironment ? 1 : 0;
@@ -899,6 +950,7 @@
 
 #if PY_MAJOR_VERSION >= 3 && PY_MINOR_VERSION >= 11
 		ModulePyConfig.isolated = IsolatedInterpreterFlag;
+		ModulePyConfig.stdio_encoding = Utf8String.GetData();
 #elif PY_MAJOR_VERSION >= 3 && PY_MINOR_VERSION >= 4
 		Py_IsolatedFlag = IsolatedInterpreterFlag; // If not zero, sys.path contains neither the script's directory nor the user's site-packages directory.
 		Py_SetStandardStreamEncoding("utf-8", nullptr);
@@ -927,40 +979,10 @@
 			PyEval_InitThreads();
 		}
 #endif // PY_MAJOR_VERSION >= 3 && PY_MINOR_VERSION < 7
-
-#if PLATFORM_WINDOWS
-		// We call _setmode here to restore the previous state
-		if (StdInMode != -1)
-		{
-			CA_SUPPRESS(6031)
-			fflush(stdin);
-			CA_SUPPRESS(6031)
-			_setmode(_fileno(stdin), StdInMode);
-		}
-		if (StdOutMode != -1)
-		{
-			CA_SUPPRESS(6031)
-			fflush(stdout);
-			CA_SUPPRESS(6031)
-			_setmode(_fileno(stdout), StdOutMode);
-		}
-		if (StdErrMode != -1)
-		{
-			CA_SUPPRESS(6031)
-			fflush(stderr);
-			CA_SUPPRESS(6031)
-			_setmode(_fileno(stderr), StdErrMode);
-		}
-#endif	// PLATFORM_WINDOWS
-
-#if PY_MAJOR_VERSION >= 3 && PY_MINOR_VERSION >= 7
-		// We call setlocale here to restore the previous state
-		if (!CurrentLocale.IsEmpty())
-		{
-			setlocale(LC_ALL, TCHAR_TO_ANSI(*CurrentLocale));
-		}
-#endif	// PY_MAJOR_VERSION >= 3 && PY_MINOR_VERSION >= 7
-
+	}
+	
+	// Setup UE conventions for the embedded interpreter environment
+	{
 #if PY_MAJOR_VERSION >=3 && PY_MINOR_VERSION >= 11
 		// Set default argv to [""]
 		FPyObjectPtr PyArgvList = MakeEmptyArgvList();
@@ -1366,7 +1388,8 @@
 		const FString PotentialFilePath = PySysPath / TEXT("init_unreal.py");
 		if (FPaths::FileExists(PotentialFilePath))
 		{
-			Progress.EnterProgressFrame(1.0f, FText::Format(LOCTEXT("PythonScriptPluginInitScripts_Running", "Running {0}..."), FText::FromString(PotentialFilePath)));
+			FText StartupScriptInfoText = FText::Format(LOCTEXT("PythonScriptPluginInitScripts_Running", "Running start-up script {0}..."), FText::FromString(PotentialFilePath));
+			Progress.EnterProgressFrame(1.0f, StartupScriptInfoText);
 			Progress.ForceRefresh();
 
 			// Execute these files in the "public" scope, as if their contents had been run directly in the console
@@ -1374,7 +1397,7 @@
 			FPythonCommandEx InitUnrealPythonCommand;
 			InitUnrealPythonCommand.FileExecutionScope = EPythonFileExecutionScope::Public;
 			
-			FScopedDurationTimeLogger ScopedTimer(FString::Printf(TEXT("Running Python start-up script '%s'"), *PotentialFilePath));
+			UE_SCOPED_TIMER(*StartupScriptInfoText.ToString(), LogPython, Display);
 			RunFile(*PotentialFilePath, *InitUnrealPythonCommand.Command, InitUnrealPythonCommand);
 		}
 		else
@@ -1385,9 +1408,10 @@
 
 	for (const FString& StartupScript : GetDefault<UPythonScriptPluginSettings>()->StartupScripts)
 	{
-		Progress.EnterProgressFrame(1.0f, FText::Format(LOCTEXT("PythonScriptPluginInitScripts_Running", "Running {0}..."), FText::FromString(StartupScript)));
-		
-		FScopedDurationTimeLogger ScopedTimer(FString::Printf(TEXT("Running Python start-up script '%s'"), *StartupScript));
+		FText StartupScriptInfoText = FText::Format(LOCTEXT("PythonScriptPluginInitScripts_Running", "Running start-up script {0}..."), FText::FromString(StartupScript));
+		Progress.EnterProgressFrame(1.0f, StartupScriptInfoText);
+
+		UE_SCOPED_TIMER(*StartupScriptInfoText.ToString(), LogPython, Display);
 		ExecPythonCommand(*StartupScript);
 	}
 
@@ -1501,6 +1525,7 @@
 		return nullptr;
 	}
 
+	FScopedEncodingGuard EncodingGuard;
 	PyUtil::FEvalStack::Get().PushContext(PyUtil::FEvalStack::FEvalContext{ InContext, InGlobalDict, InLocalDict});
 	PyObject* PyEvalResult = PyEval_EvalCode((PyUtil::FPyCodeObjectType*)PyCodeObj.Get(), InGlobalDict, InLocalDict);
 	PyUtil::FEvalStack::Get().PopContext();
