Fixed pip install progress no longer showing on the editor splash screen

#rnx

diff --git a/Engine/Plugins/Experimental/PythonScriptPlugin/Source/PythonScriptPlugin/Private/PipInstall.cpp b/Engine/Plugins/Experimental/PythonScriptPlugin/Source/PythonScriptPlugin/Private/PipInstall.cpp
--- a/Engine/Plugins/Experimental/PythonScriptPlugin/Source/PythonScriptPlugin/Private/PipInstall.cpp
+++ b/Engine/Plugins/Experimental/PythonScriptPlugin/Source/PythonScriptPlugin/Private/PipInstall.cpp
@@ -633,6 +633,7 @@ bool FPipInstall::RunLoggedSubprocess(int32* OutExitCode, const FText& Descripti
 	int AmountOfWork = (CmdParser.IsValid()) ? CmdParser->GetTotalWork() : 0;
 
 	FScopedSlowTask SubprocessTask(AmountOfWork, Description, true, *Context);
+	SubprocessTask.Visibility = ESlowTaskVisibility::Important; // this function can be very slow, users will benefit from our messages
 	SubprocessTask.MakeDialog();
 
 	// Create a read and write pipe for the child process
