Fixed PythonGeneratedClass failing to build DestructorLink and PostConstructLink due to being marked as CLASS_Native

This would lead to a crash when trying to use a non-zeroable property type (like FText) on a Python generated class, as the instance would be left zeroed rather than constructed. It would also leak memory as the properties on the class were never destroyed.

diff --git a/Engine/Plugins/Experimental/PythonScriptPlugin/Source/PythonScriptPlugin/Private/PyWrapperObject.cpp b/Engine/Plugins/Experimental/PythonScriptPlugin/Source/PythonScriptPlugin/Private/PyWrapperObject.cpp
--- a/Engine/Plugins/Experimental/PythonScriptPlugin/Source/PythonScriptPlugin/Private/PyWrapperObject.cpp
+++ b/Engine/Plugins/Experimental/PythonScriptPlugin/Source/PythonScriptPlugin/Private/PyWrapperObject.cpp
@@ -1952,6 +1952,54 @@ void UPythonGeneratedClass::PostRename(UObject* OldOuter, const FName OldName)
 	}
 }
 
+void UPythonGeneratedClass::Link(FArchive& Ar, bool bRelinkExistingProperties)
+{
+	Super::Link(Ar, bRelinkExistingProperties);
+
+	// Note: UPythonGeneratedClass instances are currently (incorrectly) marked as CLASS_Native, so we have to manually build the DestructorLink and 
+	//       PostConstructLink chains as the base UClass::Link skips them since it assumes C++ will initialize/destroy them
+	
+	// UClass::Link may have already added some properties to these linked lists, so find the end of the list and build a set to avoid adding the same property twice
+	FProperty** DestructorLinkPtr = &DestructorLink;
+	TSet<const FProperty*> BaseDestructorLinkProperties;
+	while (*DestructorLinkPtr)
+	{
+		BaseDestructorLinkProperties.Add(*DestructorLinkPtr);
+		DestructorLinkPtr = &(*DestructorLinkPtr)->DestructorLinkNext;
+	}
+	FProperty** PostConstructLinkPtr = &PostConstructLink;
+	TSet<const FProperty*> BasePostConstructLinkProperties;
+	while (*PostConstructLinkPtr)
+	{
+		BasePostConstructLinkProperties.Add(*PostConstructLinkPtr);
+		PostConstructLinkPtr = &(*PostConstructLinkPtr)->PostConstructLinkNext;
+	}
+	for (TFieldIterator<FProperty> It(this); It; ++It)
+	{
+		FProperty* Property = *It;
+
+		const UPythonGeneratedClass* OwnerClass = Cast<UPythonGeneratedClass>(Property->GetOwnerClass());
+		if (!OwnerClass)
+		{
+			continue;
+		}
+
+		if (!Property->HasAnyPropertyFlags(CPF_IsPlainOldData | CPF_NoDestructor) && !BaseDestructorLinkProperties.Contains(Property))
+		{
+			*DestructorLinkPtr = Property;
+			DestructorLinkPtr = &(*DestructorLinkPtr)->DestructorLinkNext;
+		}
+
+		if (!BasePostConstructLinkProperties.Contains(Property))
+		{
+			*PostConstructLinkPtr = Property;
+			PostConstructLinkPtr = &(*PostConstructLinkPtr)->PostConstructLinkNext;
+		}
+	}
+	*DestructorLinkPtr = nullptr;
+	*PostConstructLinkPtr = nullptr;
+}
+
 void UPythonGeneratedClass::PostInitInstance(UObject* InObj, FObjectInstancingGraph* InstanceGraph)
 {
 	Super::PostInitInstance(InObj, InstanceGraph);
@@ -2102,7 +2150,8 @@ bool UPythonGeneratedClass::ReparentDerivedClasses(UPythonGeneratedClass* InOldP
 
 	for (UClass* DerivedClass : DerivedClasses)
 	{
-		if (DerivedClass->HasAnyClassFlags(CLASS_Native | CLASS_NewerVersionExists))
+		// Note: Not checking CLASS_Native below as UPythonGeneratedClass instances are currently (incorrectly) marked as CLASS_Native
+		if (DerivedClass->HasAnyClassFlags(/*CLASS_Native | */CLASS_NewerVersionExists))
 		{
 			continue;
 		}
diff --git a/Engine/Plugins/Experimental/PythonScriptPlugin/Source/PythonScriptPlugin/Private/PyWrapperObject.h b/Engine/Plugins/Experimental/PythonScriptPlugin/Source/PythonScriptPlugin/Private/PyWrapperObject.h
--- a/Engine/Plugins/Experimental/PythonScriptPlugin/Source/PythonScriptPlugin/Private/PyWrapperObject.h
+++ b/Engine/Plugins/Experimental/PythonScriptPlugin/Source/PythonScriptPlugin/Private/PyWrapperObject.h
@@ -200,6 +200,7 @@ public:
 	}
 
 	//~ UClass interface
+	virtual void Link(FArchive& Ar, bool bRelinkExistingProperties) override;
 	virtual void PostInitInstance(UObject* InObj, FObjectInstancingGraph* InstanceGraph) override;
 
 	//~ IPythonResourceOwner interface
