[VFX Reference Update] Fix breakage in generating the python stub file.

- Consider 'self' a reserved name when creating Python glue for property names and function parameters
- Fix a name collision in a (previously) experimental PCG function

diff --git a/Engine/Plugins/Experimental/PythonScriptPlugin/Source/PythonScriptPlugin/Private/PyGenUtil.cpp b/Engine/Plugins/Experimental/PythonScriptPlugin/Source/PythonScriptPlugin/Private/PyGenUtil.cpp
--- a/Engine/Plugins/Experimental/PythonScriptPlugin/Source/PythonScriptPlugin/Private/PyGenUtil.cpp
+++ b/Engine/Plugins/Experimental/PythonScriptPlugin/Source/PythonScriptPlugin/Private/PyGenUtil.cpp
@@ -879,9 +879,15 @@ void ExtractFunctionParams(const UFunction* InFunc, TArray<FGeneratedWrappedMeth
 	auto AddGeneratedWrappedMethodParameter = [InFunc](const FProperty* InParam, TArray<FGeneratedWrappedMethodParameter>& OutParams)
 	{
 		const FString ParamName = InParam->GetName();
-		const FString PythonParamName = PythonizePropertyName(ParamName, EPythonizeNameCase::Lower);
 		const FName DefaultValueMetaDataKey = *FString::Printf(TEXT("CPP_Default_%s"), *ParamName);
 
+		FString PythonParamName = PythonizePropertyName(ParamName, EPythonizeNameCase::Lower);
+		if (!InFunc->HasAnyFunctionFlags(FUNC_Static) && PythonParamName == TEXTVIEW("self"))
+		{
+			// self is a reserved function parameter name
+			PythonParamName.InsertAt(0, TEXT('_'));
+		}
+
 		FGeneratedWrappedMethodParameter& GeneratedWrappedMethodParam = OutParams.AddDefaulted_GetRef();
 		GeneratedWrappedMethodParam.ParamName = TCHARToUTF8Buffer(*PythonParamName);
 		GeneratedWrappedMethodParam.ParamProp = InParam;
@@ -1850,18 +1856,17 @@ FString PythonizePropertyName(FStringView InName, const EPythonizeNameCase InNam
 			continue;
 		}
 
-		// Strip the "Out" prefix from names
-		//if (InName.Len() - NameOffset >= 4 && InName[NameOffset] == TEXT('O') && InName[NameOffset + 1] == TEXT('u') && InName[NameOffset + 2] == TEXT('t') && FChar::IsUpper(InName[NameOffset + 3]))
-		//{
-		//	NameOffset += 3;
-		//	continue;
-		//}
-
 		// Nothing more to strip
 		break;
 	}
 
-	return PythonizeName(NameOffset ? InName.RightChop(NameOffset) : InName, InNameCase);
+	FString PythonPropertyName = PythonizeName(NameOffset ? InName.RightChop(NameOffset) : InName, InNameCase);
+	if (PythonPropertyName == TEXTVIEW("self"))
+	{
+		PythonPropertyName.InsertAt(0, TEXT('_'));
+	}
+
+	return PythonPropertyName;
 }
 
 FString PythonizePropertyTooltip(const FParsedTooltip& InTooltip, const FProperty* InProp, const uint64 InReadOnlyFlags)
diff --git a/Engine/Plugins/PCG/Source/PCG/Private/Metadata/PCGAttributePropertySelector.cpp b/Engine/Plugins/PCG/Source/PCG/Private/Metadata/PCGAttributePropertySelector.cpp
--- a/Engine/Plugins/PCG/Source/PCG/Private/Metadata/PCGAttributePropertySelector.cpp
+++ b/Engine/Plugins/PCG/Source/PCG/Private/Metadata/PCGAttributePropertySelector.cpp
@@ -492,7 +492,7 @@ FPCGAttributePropertyInputSelector UPCGAttributePropertySelectorBlueprintHelpers
 	return Selector.CopyAndFixLast(InData);
 }
 
-FPCGAttributePropertyOutputSelector UPCGAttributePropertySelectorBlueprintHelpers::CopyAndFixSource(const FPCGAttributePropertyOutputSelector& Selector, const FPCGAttributePropertyInputSelector& InSelector)
+FPCGAttributePropertyOutputSelector UPCGAttributePropertySelectorBlueprintHelpers::CopyAndFixSource(const FPCGAttributePropertyOutputSelector& OutputSelector, const FPCGAttributePropertyInputSelector& InputSelector)
 {
-	return Selector.CopyAndFixSource(&InSelector);
+	return OutputSelector.CopyAndFixSource(&InputSelector);
 }
\ No newline at end of file
diff --git a/Engine/Plugins/PCG/Source/PCG/Public/Metadata/PCGAttributePropertySelector.h b/Engine/Plugins/PCG/Source/PCG/Public/Metadata/PCGAttributePropertySelector.h
--- a/Engine/Plugins/PCG/Source/PCG/Public/Metadata/PCGAttributePropertySelector.h
+++ b/Engine/Plugins/PCG/Source/PCG/Public/Metadata/PCGAttributePropertySelector.h
@@ -262,5 +262,5 @@ public:
 	static FPCGAttributePropertyInputSelector CopyAndFixLast(UPARAM(const, ref) const FPCGAttributePropertyInputSelector& Selector, const UPCGData* InData);
 
 	UFUNCTION(BlueprintCallable, Category = "PCG|PCGAttributePropertySelector", meta = (ScriptMethod))
-	static FPCGAttributePropertyOutputSelector CopyAndFixSource(UPARAM(const, ref) const FPCGAttributePropertyOutputSelector& Selector, const FPCGAttributePropertyInputSelector& InSelector);
+	static FPCGAttributePropertyOutputSelector CopyAndFixSource(UPARAM(const, ref) const FPCGAttributePropertyOutputSelector& OutputSelector, const FPCGAttributePropertyInputSelector& InputSelector);
 };
\ No newline at end of file
