[Python] Update the documentation writer to check the python stub when -NoHTML is passed to allow for faster stub validation

Also remove some duplicate code that snuck through in a previous submit

diff --git a/Engine/Plugins/Experimental/PythonScriptPlugin/Source/PythonScriptPlugin/Private/PyGenUtil.cpp b/Engine/Plugins/Experimental/PythonScriptPlugin/Source/PythonScriptPlugin/Private/PyGenUtil.cpp
--- a/Engine/Plugins/Experimental/PythonScriptPlugin/Source/PythonScriptPlugin/Private/PyGenUtil.cpp
+++ b/Engine/Plugins/Experimental/PythonScriptPlugin/Source/PythonScriptPlugin/Private/PyGenUtil.cpp
@@ -879,14 +879,8 @@ void ExtractFunctionParams(const UFunction* InFunc, TArray<FGeneratedWrappedMeth
 	auto AddGeneratedWrappedMethodParameter = [InFunc](const FProperty* InParam, TArray<FGeneratedWrappedMethodParameter>& OutParams)
 	{
 		const FString ParamName = InParam->GetName();
-		const FName DefaultValueMetaDataKey = *FString::Printf(TEXT("CPP_Default_%s"), *ParamName);
-
 		FString PythonParamName = PythonizePropertyName(ParamName, EPythonizeNameCase::Lower);
-		if (!InFunc->HasAnyFunctionFlags(FUNC_Static) && PythonParamName == TEXTVIEW("self"))
-		{
-			// self is a reserved function parameter name
-			PythonParamName.InsertAt(0, TEXT('_'));
-		}
+		const FName DefaultValueMetaDataKey = *FString::Printf(TEXT("CPP_Default_%s"), *ParamName);
 
 		FGeneratedWrappedMethodParameter& GeneratedWrappedMethodParam = OutParams.AddDefaulted_GetRef();
 		GeneratedWrappedMethodParam.ParamName = TCHARToUTF8Buffer(*PythonParamName);
