Don't crash when passing None as the class type to unreal.ObjectIterator

Instead gracefully handle it and emit a type error: 'type' cannot be 'None'

diff --git a/Engine/Plugins/Experimental/PythonScriptPlugin/Source/PythonScriptPlugin/Private/PyCore.cpp b/Engine/Plugins/Experimental/PythonScriptPlugin/Source/PythonScriptPlugin/Private/PyCore.cpp
--- a/Engine/Plugins/Experimental/PythonScriptPlugin/Source/PythonScriptPlugin/Private/PyCore.cpp
+++ b/Engine/Plugins/Experimental/PythonScriptPlugin/Source/PythonScriptPlugin/Private/PyCore.cpp
@@ -728,10 +728,18 @@ PyTypeObject InitializePyObjectIteratorType()
 			}
 
 			UClass* IterClass = UObject::StaticClass();
-			if (PyTypeObj && !PyConversion::NativizeClass(PyTypeObj, IterClass, nullptr))
+			if (PyTypeObj)
 			{
-				PyUtil::SetPythonError(PyExc_TypeError, InSelf, *FString::Printf(TEXT("Failed to convert 'type' (%s) to 'Class'"), *PyUtil::GetFriendlyTypename(PyTypeObj)));
-				return -1;
+				if (PyTypeObj == Py_None)
+				{
+					PyUtil::SetPythonError(PyExc_Exception, InSelf, TEXT("'type' cannot be 'None'"));
+					return -1;
+				}
+				if (!PyConversion::NativizeClass(PyTypeObj, IterClass, nullptr))
+				{
+					PyUtil::SetPythonError(PyExc_TypeError, InSelf, *FString::Printf(TEXT("Failed to convert 'type' (%s) to 'Class'"), *PyUtil::GetFriendlyTypename(PyTypeObj)));
+					return -1;
+				}
 			}
 
 			return FPyObjectIterator::Init(InSelf, IterClass, nullptr);
