Remove Native/BlueprintType flags from Python generated types

Remove CLASS_Native from PythonGeneratedClass so that the default UClass::Link works correctly.
Remove BlueprintType meta-data from PythonGeneratedStruct/PythonGeneratedEnum, as these types are transient and don't serialize correctly in a BP graph anyway.

diff --git a/Engine/Plugins/Experimental/PythonScriptPlugin/Source/PythonScriptPlugin/Private/PyWrapperEnum.cpp b/Engine/Plugins/Experimental/PythonScriptPlugin/Source/PythonScriptPlugin/Private/PyWrapperEnum.cpp
--- a/Engine/Plugins/Experimental/PythonScriptPlugin/Source/PythonScriptPlugin/Private/PyWrapperEnum.cpp
+++ b/Engine/Plugins/Experimental/PythonScriptPlugin/Source/PythonScriptPlugin/Private/PyWrapperEnum.cpp
@@ -758,7 +758,6 @@ public:
 		{
 			NewEnum = NewObject<UPythonGeneratedEnum>(EnumOuter, *EnumName, RF_Public | RF_Transient);
 			NewEnum->SetMetaData(TEXT("DisplayName"), *PyUtil::GetGeneratedTypeDisplayName(PyType));
-			NewEnum->SetMetaData(TEXT("BlueprintType"), TEXT("true"));
 		}
 		NewEnum->EnumValueDefs.Reset();
 
diff --git a/Engine/Plugins/Experimental/PythonScriptPlugin/Source/PythonScriptPlugin/Private/PyWrapperObject.cpp b/Engine/Plugins/Experimental/PythonScriptPlugin/Source/PythonScriptPlugin/Private/PyWrapperObject.cpp
--- a/Engine/Plugins/Experimental/PythonScriptPlugin/Source/PythonScriptPlugin/Private/PyWrapperObject.cpp
+++ b/Engine/Plugins/Experimental/PythonScriptPlugin/Source/PythonScriptPlugin/Private/PyWrapperObject.cpp
@@ -1334,7 +1334,6 @@ public:
 		NewClass = NewObject<UPythonGeneratedClass>(ClassOuter, *NewClassName, RF_Public | RF_Transient);
 		NewClass->SetMetaData(TEXT("DisplayName"), *PyUtil::GetGeneratedTypeDisplayName(PyType));
 		NewClass->SetSuperStruct(InSuperClass);
-		NewClass->ClassFlags |= CLASS_Native;
 	}
 
 	FPythonGeneratedClassBuilder(UPythonGeneratedClass* InOldClass, UClass* InSuperClass)
@@ -1350,7 +1349,6 @@ public:
 		NewClass = NewObject<UPythonGeneratedClass>(ClassOuter, *NewClassName, RF_Public | RF_Transient);
 		NewClass->SetMetaData(TEXT("DisplayName"), *PyUtil::GetGeneratedTypeDisplayName(PyType));
 		NewClass->SetSuperStruct(InSuperClass);
-		NewClass->ClassFlags |= CLASS_Native;
 	}
 
 	~FPythonGeneratedClassBuilder()
@@ -1952,54 +1950,6 @@ void UPythonGeneratedClass::PostRename(UObject* OldOuter, const FName OldName)
 	}
 }
 
-void UPythonGeneratedClass::Link(FArchive& Ar, bool bRelinkExistingProperties)
-{
-	Super::Link(Ar, bRelinkExistingProperties);
-
-	// Note: UPythonGeneratedClass instances are currently (incorrectly) marked as CLASS_Native, so we have to manually build the DestructorLink and 
-	//       PostConstructLink chains as the base UClass::Link skips them since it assumes C++ will initialize/destroy them
-	
-	// UClass::Link may have already added some properties to these linked lists, so find the end of the list and build a set to avoid adding the same property twice
-	FProperty** DestructorLinkPtr = &DestructorLink;
-	TSet<const FProperty*> BaseDestructorLinkProperties;
-	while (*DestructorLinkPtr)
-	{
-		BaseDestructorLinkProperties.Add(*DestructorLinkPtr);
-		DestructorLinkPtr = &(*DestructorLinkPtr)->DestructorLinkNext;
-	}
-	FProperty** PostConstructLinkPtr = &PostConstructLink;
-	TSet<const FProperty*> BasePostConstructLinkProperties;
-	while (*PostConstructLinkPtr)
-	{
-		BasePostConstructLinkProperties.Add(*PostConstructLinkPtr);
-		PostConstructLinkPtr = &(*PostConstructLinkPtr)->PostConstructLinkNext;
-	}
-	for (TFieldIterator<FProperty> It(this); It; ++It)
-	{
-		FProperty* Property = *It;
-
-		const UPythonGeneratedClass* OwnerClass = Cast<UPythonGeneratedClass>(Property->GetOwnerClass());
-		if (!OwnerClass)
-		{
-			continue;
-		}
-
-		if (!Property->HasAnyPropertyFlags(CPF_IsPlainOldData | CPF_NoDestructor) && !BaseDestructorLinkProperties.Contains(Property))
-		{
-			*DestructorLinkPtr = Property;
-			DestructorLinkPtr = &(*DestructorLinkPtr)->DestructorLinkNext;
-		}
-
-		if (!BasePostConstructLinkProperties.Contains(Property))
-		{
-			*PostConstructLinkPtr = Property;
-			PostConstructLinkPtr = &(*PostConstructLinkPtr)->PostConstructLinkNext;
-		}
-	}
-	*DestructorLinkPtr = nullptr;
-	*PostConstructLinkPtr = nullptr;
-}
-
 void UPythonGeneratedClass::PostInitInstance(UObject* InObj, FObjectInstancingGraph* InstanceGraph)
 {
 	Super::PostInitInstance(InObj, InstanceGraph);
@@ -2150,8 +2100,7 @@ bool UPythonGeneratedClass::ReparentDerivedClasses(UPythonGeneratedClass* InOldP
 
 	for (UClass* DerivedClass : DerivedClasses)
 	{
-		// Note: Not checking CLASS_Native below as UPythonGeneratedClass instances are currently (incorrectly) marked as CLASS_Native
-		if (DerivedClass->HasAnyClassFlags(/*CLASS_Native | */CLASS_NewerVersionExists))
+		if (DerivedClass->HasAnyClassFlags(CLASS_NewerVersionExists))
 		{
 			continue;
 		}
diff --git a/Engine/Plugins/Experimental/PythonScriptPlugin/Source/PythonScriptPlugin/Private/PyWrapperObject.h b/Engine/Plugins/Experimental/PythonScriptPlugin/Source/PythonScriptPlugin/Private/PyWrapperObject.h
--- a/Engine/Plugins/Experimental/PythonScriptPlugin/Source/PythonScriptPlugin/Private/PyWrapperObject.h
+++ b/Engine/Plugins/Experimental/PythonScriptPlugin/Source/PythonScriptPlugin/Private/PyWrapperObject.h
@@ -183,7 +183,7 @@ typedef TPyPtr<FPyWrapperObject> FPyWrapperObjectPtr;
 #endif	// WITH_PYTHON
 
 /** An Unreal class that was generated from a Python type */
-UCLASS(BlueprintType, Transient)
+UCLASS(Transient)
 class UPythonGeneratedClass final : public UClass, public IPythonResourceOwner
 {
 	GENERATED_BODY()
@@ -200,7 +200,6 @@ public:
 	}
 
 	//~ UClass interface
-	virtual void Link(FArchive& Ar, bool bRelinkExistingProperties) override;
 	virtual void PostInitInstance(UObject* InObj, FObjectInstancingGraph* InstanceGraph) override;
 
 	//~ IPythonResourceOwner interface
diff --git a/Engine/Plugins/Experimental/PythonScriptPlugin/Source/PythonScriptPlugin/Private/PyWrapperStruct.cpp b/Engine/Plugins/Experimental/PythonScriptPlugin/Source/PythonScriptPlugin/Private/PyWrapperStruct.cpp
--- a/Engine/Plugins/Experimental/PythonScriptPlugin/Source/PythonScriptPlugin/Private/PyWrapperStruct.cpp
+++ b/Engine/Plugins/Experimental/PythonScriptPlugin/Source/PythonScriptPlugin/Private/PyWrapperStruct.cpp
@@ -1456,7 +1456,6 @@ public:
 		const FString NewStructName = MakeUniqueObjectName(StructOuter, UPythonGeneratedStruct::StaticClass(), *FString::Printf(TEXT("%s_NEWINST"), *StructName)).ToString();
 		NewStruct = NewObject<UPythonGeneratedStruct>(StructOuter, *NewStructName, RF_Public | RF_Transient);
 		NewStruct->SetMetaData(TEXT("DisplayName"), *PyUtil::GetGeneratedTypeDisplayName(PyType));
-		NewStruct->SetMetaData(TEXT("BlueprintType"), TEXT("true"));
 		NewStruct->SetSuperStruct(InSuperStruct);
 	}
 
