Re-add BlueprintType to Python generated types

This allows split struct pins from a Python defined BPFL to continue to work in BPs

diff --git a/Engine/Plugins/Experimental/PythonScriptPlugin/Source/PythonScriptPlugin/Private/PyWrapperEnum.cpp b/Engine/Plugins/Experimental/PythonScriptPlugin/Source/PythonScriptPlugin/Private/PyWrapperEnum.cpp
--- a/Engine/Plugins/Experimental/PythonScriptPlugin/Source/PythonScriptPlugin/Private/PyWrapperEnum.cpp
+++ b/Engine/Plugins/Experimental/PythonScriptPlugin/Source/PythonScriptPlugin/Private/PyWrapperEnum.cpp
@@ -758,6 +758,7 @@ public:
 		{
 			NewEnum = NewObject<UPythonGeneratedEnum>(EnumOuter, *EnumName, RF_Public | RF_Transient);
 			NewEnum->SetMetaData(TEXT("DisplayName"), *PyUtil::GetGeneratedTypeDisplayName(PyType));
+			NewEnum->SetMetaData(TEXT("BlueprintType"), TEXT("true"));
 		}
 		NewEnum->EnumValueDefs.Reset();
 
diff --git a/Engine/Plugins/Experimental/PythonScriptPlugin/Source/PythonScriptPlugin/Private/PyWrapperObject.cpp b/Engine/Plugins/Experimental/PythonScriptPlugin/Source/PythonScriptPlugin/Private/PyWrapperObject.cpp
--- a/Engine/Plugins/Experimental/PythonScriptPlugin/Source/PythonScriptPlugin/Private/PyWrapperObject.cpp
+++ b/Engine/Plugins/Experimental/PythonScriptPlugin/Source/PythonScriptPlugin/Private/PyWrapperObject.cpp
@@ -1333,6 +1333,7 @@ public:
 		const FString NewClassName = MakeUniqueObjectName(ClassOuter, UPythonGeneratedClass::StaticClass(), *FString::Printf(TEXT("%s_NEWINST"), *ClassName)).ToString();
 		NewClass = NewObject<UPythonGeneratedClass>(ClassOuter, *NewClassName, RF_Public | RF_Transient);
 		NewClass->SetMetaData(TEXT("DisplayName"), *PyUtil::GetGeneratedTypeDisplayName(PyType));
+		NewClass->SetMetaData(TEXT("BlueprintType"), TEXT("true"));
 		NewClass->SetSuperStruct(InSuperClass);
 	}
 
@@ -1348,6 +1349,7 @@ public:
 		const FString NewClassName = MakeUniqueObjectName(ClassOuter, UPythonGeneratedClass::StaticClass(), *FString::Printf(TEXT("%s_NEWINST"), *ClassName)).ToString();
 		NewClass = NewObject<UPythonGeneratedClass>(ClassOuter, *NewClassName, RF_Public | RF_Transient);
 		NewClass->SetMetaData(TEXT("DisplayName"), *PyUtil::GetGeneratedTypeDisplayName(PyType));
+		NewClass->SetMetaData(TEXT("BlueprintType"), TEXT("true"));
 		NewClass->SetSuperStruct(InSuperClass);
 	}
 
diff --git a/Engine/Plugins/Experimental/PythonScriptPlugin/Source/PythonScriptPlugin/Private/PyWrapperStruct.cpp b/Engine/Plugins/Experimental/PythonScriptPlugin/Source/PythonScriptPlugin/Private/PyWrapperStruct.cpp
--- a/Engine/Plugins/Experimental/PythonScriptPlugin/Source/PythonScriptPlugin/Private/PyWrapperStruct.cpp
+++ b/Engine/Plugins/Experimental/PythonScriptPlugin/Source/PythonScriptPlugin/Private/PyWrapperStruct.cpp
@@ -1456,6 +1456,7 @@ public:
 		const FString NewStructName = MakeUniqueObjectName(StructOuter, UPythonGeneratedStruct::StaticClass(), *FString::Printf(TEXT("%s_NEWINST"), *StructName)).ToString();
 		NewStruct = NewObject<UPythonGeneratedStruct>(StructOuter, *NewStructName, RF_Public | RF_Transient);
 		NewStruct->SetMetaData(TEXT("DisplayName"), *PyUtil::GetGeneratedTypeDisplayName(PyType));
+		NewStruct->SetMetaData(TEXT("BlueprintType"), TEXT("true"));
 		NewStruct->SetSuperStruct(InSuperStruct);
 	}
 
