Deprecated creating an instance of a BlueprintFunctionLibrary from Python

These types only contain state functions and no state, and so you should just call their classmethods directly on the class from Python

diff --git a/Engine/Plugins/Experimental/PythonScriptPlugin/Source/PythonScriptPlugin/Private/PyUtil.cpp b/Engine/Plugins/Experimental/PythonScriptPlugin/Source/PythonScriptPlugin/Private/PyUtil.cpp
--- a/Engine/Plugins/Experimental/PythonScriptPlugin/Source/PythonScriptPlugin/Private/PyUtil.cpp
+++ b/Engine/Plugins/Experimental/PythonScriptPlugin/Source/PythonScriptPlugin/Private/PyUtil.cpp
@@ -18,22 +18,23 @@
 #include "PyWrapperFieldPath.h"
 #include "PyWrapperTypeRegistry.h"
 
+#include "HAL/FileManager.h"
+#include "Kismet/BlueprintFunctionLibrary.h"
+#include "Misc/DefaultValueHelper.h"
+#include "Misc/MessageDialog.h"
 #include "Misc/PackageName.h"
 #include "Misc/Paths.h"
 #include "Misc/PathViews.h"
 #include "Misc/ScopeExit.h"
-#include "Misc/MessageDialog.h"
-#include "Misc/DefaultValueHelper.h"
 #include "ProfilingDebugging/CpuProfilerTrace.h"
-#include "HAL/FileManager.h"
-#include "UObject/Package.h"
-#include "UObject/UnrealType.h"
-#include "UObject/EnumProperty.h"
-#include "UObject/TextProperty.h"
-#include "UObject/StructOnScope.h"
-#include "UObject/PropertyPortFlags.h"
 #include "Subsystems/EngineSubsystem.h"
 #include "Templates/Casts.h"
+#include "UObject/EnumProperty.h"
+#include "UObject/Package.h"
+#include "UObject/PropertyPortFlags.h"
+#include "UObject/StructOnScope.h"
+#include "UObject/TextProperty.h"
+#include "UObject/UnrealType.h"
 
 #if WITH_EDITOR
 #include "EditorSubsystem.h"
@@ -690,23 +691,32 @@ UObject* NewObject(UClass* InObjClass, UObject* InObjectOuter, const FName InObj
 {
 	if (InObjClass)
 	{
-		if (InObjClass->IsChildOf(UEngineSubsystem::StaticClass()))
+		if (InObjClass->IsChildOf<UBlueprintFunctionLibrary>())
+		{
+			// Starting with UE 5.5, generates a warning/deprecation message.
+			SetPythonWarning(PyExc_DeprecationWarning, InErrorCtxt, *FString::Printf(TEXT("Creating an instance of a BlueprintFunctionLibrary has been deprecated since UE 5.5 and will be removed in the future. Call its classmethods directly on the class, eg, 'unreal.%s.foo()'."), *PyGenUtil::GetClassPythonName(InObjClass)));
+
+			// For UE 5.7 or later, generate an hard error.
+			//SetPythonError(PyExc_Exception, InErrorCtxt, *FString::Printf(TEXT("Cannot create an instance of a BlueprintFunctionLibrary. Call its classmethods directly on the class, eg, 'unreal.%s.foo()'."), *PyGenUtil::GetClassPythonName(InObjClass)));
+			//return nullptr;
+		}
+		else if (InObjClass->IsChildOf<UEngineSubsystem>())
 		{
 			// Starting with UE 5.2, generates a warning/deprecation message.
-			SetPythonWarning(PyExc_DeprecationWarning, InErrorCtxt, *FString::Printf(TEXT("Engine subsystems creation is deprecated and will be removed in UE 5.3. Use 'unreal.get_engine_subsystem(unreal.%s)' to get an instance of the subsystem."), *PyGenUtil::GetClassPythonName(InObjClass)));
+			SetPythonWarning(PyExc_DeprecationWarning, InErrorCtxt, *FString::Printf(TEXT("Creating an instance of an Engine subsystem has been deprecated since UE 5.2 and will be removed in the future. Use 'unreal.get_engine_subsystem(unreal.%s)' to get the subsystem instance."), *PyGenUtil::GetClassPythonName(InObjClass)));
 
 			// For UE 5.3 or later, generate an hard error.
-			//SetPythonError(PyExc_Exception, InErrorCtxt, *FString::Printf(TEXT("Engine subsystems cannot be created. Use 'unreal.get_engine_subsystem(unreal.%s)' to get an instance of the subsystem."), *PyGenUtil::GetClassPythonName(InObjClass)));
+			//SetPythonError(PyExc_Exception, InErrorCtxt, *FString::Printf(TEXT("Cannot create an instance of an Engine subsystem. Use 'unreal.get_engine_subsystem(unreal.%s)' to get the subsystem instance."), *PyGenUtil::GetClassPythonName(InObjClass)));
 			//return nullptr;
 		}
 #if WITH_EDITOR
-		else if (InObjClass->IsChildOf(UEditorSubsystem::StaticClass()))
+		else if (InObjClass->IsChildOf<UEditorSubsystem>())
 		{
 			// Starting with UE 5.2, generates a warning/deprecation message.
-			SetPythonWarning(PyExc_DeprecationWarning, InErrorCtxt, *FString::Printf(TEXT("Editor subsystems creation is deprecated and will be removed in UE 5.3. Use 'unreal.get_editor_subsystem(unreal.%s)' to get an instance of the subsystem."), *PyGenUtil::GetClassPythonName(InObjClass)));
+			SetPythonWarning(PyExc_DeprecationWarning, InErrorCtxt, *FString::Printf(TEXT("Creating an instance of an Editor subsystem has been deprecated since UE 5.2 and will be removed in the future. Use 'unreal.get_editor_subsystem(unreal.%s)' to get the subsystem instance."), *PyGenUtil::GetClassPythonName(InObjClass)));
 
 			// For UE 5.3 or later, generate an hard error.
-			//SetPythonError(PyExc_Exception, InErrorCtxt, *FString::Printf(TEXT("Editor subsystems cannot be created. Use 'unreal.get_editor_subsystem(unreal.%s)' to get an instance of the subsystem."), *PyGenUtil::GetClassPythonName(InObjClass)));
+			//SetPythonError(PyExc_Exception, InErrorCtxt, *FString::Printf(TEXT("Cannot create an instance of an Editor subsystem. Use 'unreal.get_editor_subsystem(unreal.%s)' to get the subsystem instance."), *PyGenUtil::GetClassPythonName(InObjClass)));
 			//return nullptr;
 		}
 #endif
