The Execute Python Script node now return false as its success value when the python script emits an error.

diff --git a/Engine/Plugins/Experimental/PythonScriptPlugin/Source/PythonScriptPlugin/Private/PythonScriptLibrary.cpp b/Engine/Plugins/Experimental/PythonScriptPlugin/Source/PythonScriptPlugin/Private/PythonScriptLibrary.cpp
--- a/Engine/Plugins/Experimental/PythonScriptPlugin/Source/PythonScriptPlugin/Private/PythonScriptLibrary.cpp
+++ b/Engine/Plugins/Experimental/PythonScriptPlugin/Source/PythonScriptPlugin/Private/PythonScriptLibrary.cpp
@@ -160,7 +160,8 @@ DEFINE_FUNCTION(UPythonScriptLibrary::execExecutePythonScript)
 		}
 		else
 		{
-			*(bool*)RESULT_PARAM = PyUtil::ReThrowPythonError();
+			PyUtil::ReThrowPythonError();
+			*(bool*)RESULT_PARAM = false;
 		}
 	}
 #else	// WITH_PYTHON
