Only allow exporting top-level class/struct/enum types to Python

Ignore internal types, such as the sparse data structs generated within an ABPGC

#rnx

diff --git a/Engine/Plugins/Experimental/PythonScriptPlugin/Source/PythonScriptPlugin/Private/PyGenUtil.cpp b/Engine/Plugins/Experimental/PythonScriptPlugin/Source/PythonScriptPlugin/Private/PyGenUtil.cpp
--- a/Engine/Plugins/Experimental/PythonScriptPlugin/Source/PythonScriptPlugin/Private/PyGenUtil.cpp
+++ b/Engine/Plugins/Experimental/PythonScriptPlugin/Source/PythonScriptPlugin/Private/PyGenUtil.cpp
@@ -1650,7 +1650,7 @@ TArray<const UClass*> GetExportedInterfacesForClass(const UClass* InClass)
 
 bool ShouldExportClass(const UClass* InClass)
 {
-	return IsScriptExposedClass(InClass) || HasScriptExposedFields(InClass);
+	return InClass->GetOuter()->IsA<UPackage>() && (IsScriptExposedClass(InClass) || HasScriptExposedFields(InClass));
 }
 
 bool ShouldExportStruct(const UScriptStruct* InStruct)
@@ -1671,12 +1671,12 @@ bool ShouldExportStruct(const UScriptStruct* InStruct)
 	}
 	*/
 
-	return IsScriptExposedStruct(InStruct) || HasScriptExposedFields(InStruct);
+	return InStruct->GetOuter()->IsA<UPackage>() && (IsScriptExposedStruct(InStruct) || HasScriptExposedFields(InStruct));
 }
 
 bool ShouldExportEnum(const UEnum* InEnum)
 {
-	return IsScriptExposedEnum(InEnum);
+	return InEnum->GetOuter()->IsA<UPackage>() && IsScriptExposedEnum(InEnum);
 }
 
 bool ShouldExportEnumEntry(const UEnum* InEnum, int32 InEnumEntryIndex)
