More correctly respect Python being disabled in the PythonScript commandlet

When this was originally written you could only disable Python at compile time (WITH_PYTHON = 0), but now you can also disable it at runtime.

#rnx

diff --git a/Engine/Plugins/Experimental/PythonScriptPlugin/Source/PythonScriptPlugin/Private/PythonScriptCommandlet.cpp b/Engine/Plugins/Experimental/PythonScriptPlugin/Source/PythonScriptPlugin/Private/PythonScriptCommandlet.cpp
--- a/Engine/Plugins/Experimental/PythonScriptPlugin/Source/PythonScriptPlugin/Private/PythonScriptCommandlet.cpp
+++ b/Engine/Plugins/Experimental/PythonScriptPlugin/Source/PythonScriptPlugin/Private/PythonScriptCommandlet.cpp
@@ -35,11 +35,16 @@ int32 UPythonScriptCommandlet::Main(const FString& Params)
 		return -1;
 	}
 
-#if WITH_PYTHON
+	if (!IPythonScriptPlugin::Get()->IsPythonAvailable())
 	{
-		// Tick once to ensure that any start-up scripts have been run
-		FTSTicker::GetCoreTicker().Tick(0.0f);
+		UE_LOG(LogPythonScriptCommandlet, Error, TEXT("Python script cannot run as Python support is disabled!"));
+		return -1;
+	}
+
+	// Tick once to ensure that any start-up scripts have been run
+	FTSTicker::GetCoreTicker().Tick(0.0f);
 
+	{
 		UE_LOG(LogPythonScriptCommandlet, Display, TEXT("Running Python script: %s"), *PythonScript);
 
 		FPythonCommandEx PythonCommand;
@@ -51,11 +56,7 @@ int32 UPythonScriptCommandlet::Main(const FString& Params)
 			return -1;
 		}
 	}
+
 	UE_LOG(LogPythonScriptCommandlet, Display, TEXT("Python script executed successfully"));
 	return 0;
-#else	// WITH_PYTHON
-	UE_LOG(LogPythonScriptCommandlet, Error, TEXT("Python script cannot run as the plugin was built as a stub!"));
-	return -1;
-#endif	// WITH_PYTHON
 }
-
