Add clear() method for unreal.Array in python wrapper

diff --git a/Engine/Plugins/Experimental/PythonScriptPlugin/Source/PythonScriptPlugin/Private/PyWrapperArray.cpp b/Engine/Plugins/Experimental/PythonScriptPlugin/Source/PythonScriptPlugin/Private/PyWrapperArray.cpp
--- a/Engine/Plugins/Experimental/PythonScriptPlugin/Source/PythonScriptPlugin/Private/PyWrapperArray.cpp
+++ b/Engine/Plugins/Experimental/PythonScriptPlugin/Source/PythonScriptPlugin/Private/PyWrapperArray.cpp
@@ -741,6 +741,19 @@ int FPyWrapperArray::Remove(FPyWrapperArray* InSelf, PyObject* InValue)
 	return 0;
 }
 
+int FPyWrapperArray::Clear(FPyWrapperArray* InSelf)
+{
+	if (!ValidateInternalState(InSelf))
+	{
+		return -1;
+	}
+
+	FScriptArrayHelper SelfScriptArrayHelper(InSelf->ArrayProp, InSelf->ArrayInstance);
+	SelfScriptArrayHelper.EmptyValues();
+
+	return 0;
+}
+
 int FPyWrapperArray::Reverse(FPyWrapperArray* InSelf)
 {
 	if (!ValidateInternalState(InSelf))
@@ -1320,6 +1333,16 @@ PyTypeObject InitializePyWrapperArrayType()
 			return FPyWrapperArray::Pop(InSelf, ValueIndex);
 		}
 
+		static PyObject* Clear(FPyWrapperArray* InSelf)
+		{
+			if (FPyWrapperArray::Clear(InSelf) != 0)
+			{
+				return nullptr;
+			}
+
+			Py_RETURN_NONE;
+		}
+
 		static PyObject* Remove(FPyWrapperArray* InSelf, PyObject* InArgs)
 		{
 			PyObject* PyObj = nullptr;
@@ -1407,6 +1430,7 @@ PyTypeObject InitializePyWrapperArrayType()
 		{ "index", PyCFunctionCast(&FMethods::Index), METH_VARARGS | METH_KEYWORDS, "index(self, value: _ElemType, start: int = 0, stop: int = -1) -> int -- get the index of the first matching value in this Unreal array, or raise ValueError if missing (equivalent to TArray::IndexOfByKey in C++)" },
 		{ "insert", PyCFunctionCast(&FMethods::Insert), METH_VARARGS | METH_KEYWORDS, "insert(self, index: int, value: _ElemType) -> None -- insert the given value at the given index in this Unreal array" },
 		{ "pop", PyCFunctionCast(&FMethods::Pop), METH_VARARGS, "pop(self, index: int = -1) -> _ElemType -- remove and return the value at the given index in this Unreal array, or raise IndexError if the index is out-of-bounds" },
+		{ "clear", PyCFunctionCast(&FMethods::Clear), METH_NOARGS, "clear(self) -> None -- remove all values from this Unreal array" },
 		{ "remove", PyCFunctionCast(&FMethods::Remove), METH_VARARGS, "remove(self, value: _ElemType) -> None -- remove the first matching value in this Unreal array, or raise ValueError if missing" },
 		{ "reverse", PyCFunctionCast(&FMethods::Reverse), METH_NOARGS, "reverse(self) -> None -- reverse this Unreal array in-place" },
 		{ "sort", PyCFunctionCast(&FMethods::Sort), METH_VARARGS | METH_KEYWORDS, "sort(self, key: Optional[Callable[[_ElemType], object]]=None, reverse: bool=False) -> None -- stable sort this Unreal array in-place" },
diff --git a/Engine/Plugins/Experimental/PythonScriptPlugin/Source/PythonScriptPlugin/Private/PyWrapperArray.h b/Engine/Plugins/Experimental/PythonScriptPlugin/Source/PythonScriptPlugin/Private/PyWrapperArray.h
--- a/Engine/Plugins/Experimental/PythonScriptPlugin/Source/PythonScriptPlugin/Private/PyWrapperArray.h
+++ b/Engine/Plugins/Experimental/PythonScriptPlugin/Source/PythonScriptPlugin/Private/PyWrapperArray.h
@@ -92,6 +92,9 @@ struct FPyWrapperArray : public FPyWrapperBase
 	/** Remove the given value from this container (equivalent to 'x.remove(v)' in Python) */
 	static int Remove(FPyWrapperArray* InSelf, PyObject* InValue);
 
+	/** Remove all values from this container (equivalent to 'x.clear()' in Python) */
+	static int Clear(FPyWrapperArray* InSelf);
+
 	/** Reverse this container (equivalent to 'x.reverse()' in Python) */
 	static int Reverse(FPyWrapperArray* InSelf);
 
