[Python] Make the tooltips, function names, and return parameter names consistent for the python bp utilities

diff --git a/Engine/Plugins/Experimental/PythonScriptPlugin/Source/PythonScriptPlugin/Private/K2Node_ExecutePythonScript.cpp b/Engine/Plugins/Experimental/PythonScriptPlugin/Source/PythonScriptPlugin/Private/K2Node_ExecutePythonScript.cpp
--- a/Engine/Plugins/Experimental/PythonScriptPlugin/Source/PythonScriptPlugin/Private/K2Node_ExecutePythonScript.cpp
+++ b/Engine/Plugins/Experimental/PythonScriptPlugin/Source/PythonScriptPlugin/Private/K2Node_ExecutePythonScript.cpp
@@ -62,10 +62,6 @@ void UK2Node_ExecutePythonScript::AllocateDefaultPins()
 	PythonScriptPin->AutogeneratedDefaultValue = TEXT("print('Hello World!')");
 	PythonScriptPin->DefaultValue = PythonScriptPin->AutogeneratedDefaultValue;
 
-	// Rename the result pin
-	UEdGraphPin* ResultPin = FindPinChecked(UEdGraphSchema_K2::PN_ReturnValue, EGPD_Output);
-	ResultPin->PinFriendlyName = LOCTEXT("ResultPinFriendlyName", "Success?");
-
 	// Add user-defined pins
 	auto AddUserDefinedPins = [this](const TArray<FName>& InPinNames, EEdGraphPinDirection InPinDirection)
 	{
diff --git a/Engine/Plugins/Experimental/PythonScriptPlugin/Source/PythonScriptPlugin/Private/PythonScriptLibrary.h b/Engine/Plugins/Experimental/PythonScriptPlugin/Source/PythonScriptPlugin/Private/PythonScriptLibrary.h
--- a/Engine/Plugins/Experimental/PythonScriptPlugin/Source/PythonScriptPlugin/Private/PythonScriptLibrary.h
+++ b/Engine/Plugins/Experimental/PythonScriptPlugin/Source/PythonScriptPlugin/Private/PythonScriptLibrary.h
@@ -34,7 +34,7 @@ public:
 	 * @param FileExecutionScope Controls the scope used when executing Python files.
 	 * @param CommandResult The result of running the command. On success, for EvaluateStatement mode this will be the actual result of running the command, and will be None in all other cases. On failure, this will be the error information (typically a Python exception trace).
 	 * @param LogOutput The log output captured while running the command.
-	 * @return true if the command ran successfully, false if there were errors.
+	 * @return true if the command ran successfully, false if there were errors (the output log will show the errors).
 	 */
 	UFUNCTION(BlueprintCallable, Category = "Python|Execution", meta=(DisplayName="Execute Python Command (Advanced)", AdvancedDisplay="ExecutionMode,FileExecutionScope"))
 	static bool ExecutePythonCommandEx(UPARAM(meta=(MultiLine=True)) const FString& PythonCommand, FString& CommandResult, TArray<FPythonLogOutputEntry>& LogOutput, const EPythonCommandExecutionMode ExecutionMode = EPythonCommandExecutionMode::ExecuteFile, const EPythonFileExecutionScope FileExecutionScope = EPythonFileExecutionScope::Private);
@@ -44,9 +44,9 @@ public:
 	 * @param PythonScript This literal Python code to run.
 	 * @param PythonInputs The variadic input argument names (internal; set by UK2Node_ExecutePythonScript).
 	 * @param PythonInputs The variadic output argument names (internal; set by UK2Node_ExecutePythonScript).
-	 * @return true if the script ran successfully, false if there were errors.
+	 * @return true if the command ran successfully, false if there were errors (the output log will show the errors).
 	 */
-    UFUNCTION(BlueprintCallable, CustomThunk, Category = "Python|Execution", meta=(Variadic, BlueprintInternalUseOnly="true"))
+    UFUNCTION(BlueprintCallable, CustomThunk, Category = "Python|Execution", meta=(Variadic, BlueprintInternalUseOnly="true"), DisplayName = "Execute Python Script")
     static bool ExecutePythonScript(UPARAM(meta=(MultiLine=True)) const FString& PythonScript, const TArray<FString>& PythonInputs, const TArray<FString>& PythonOutputs);
 	DECLARE_FUNCTION(execExecutePythonScript);
 };
