PipInstall: Added option to allow extra args in pip install line

diff --git a/Engine/Plugins/Experimental/PythonScriptPlugin/Source/PythonScriptPlugin/Private/PipInstall.cpp b/Engine/Plugins/Experimental/PythonScriptPlugin/Source/PythonScriptPlugin/Private/PipInstall.cpp
--- a/Engine/Plugins/Experimental/PythonScriptPlugin/Source/PythonScriptPlugin/Private/PipInstall.cpp
+++ b/Engine/Plugins/Experimental/PythonScriptPlugin/Source/PythonScriptPlugin/Private/PipInstall.cpp
@@ -611,6 +611,11 @@ FString FPipInstall::SetupPipInstallCmd(const FString& ParsedReqsFile, const TAr
 		}
 	}
 
+	if (!ScriptSettings->ExtraInstallArgs.IsEmpty())
+	{
+		Cmd += " " + ScriptSettings->ExtraInstallArgs;
+	}
+	 
 	Cmd += " -r \"" + ParsedReqsFile + "\"";
 	
 	return Cmd;
diff --git a/Engine/Plugins/Experimental/PythonScriptPlugin/Source/PythonScriptPlugin/Private/PythonScriptPluginSettings.h b/Engine/Plugins/Experimental/PythonScriptPlugin/Source/PythonScriptPlugin/Private/PythonScriptPluginSettings.h
--- a/Engine/Plugins/Experimental/PythonScriptPlugin/Source/PythonScriptPlugin/Private/PythonScriptPluginSettings.h
+++ b/Engine/Plugins/Experimental/PythonScriptPlugin/Source/PythonScriptPlugin/Private/PythonScriptPluginSettings.h
@@ -104,6 +104,12 @@ public:
 	*/
 	UPROPERTY(config, EditAnywhere, Category=PythonPipInstall, AdvancedDisplay, meta=(ConfigRestartRequired=true))
 	FString OverrideIndexURL;
+	
+	/** Additional arguments passed to main pip install call, useful to add e.g. --cert or other proxy options for restrictive firewalls
+	* 	NOTE: Do not use this to add --index-url or --extra-index-url, instead use OverrideIndexURL setting or ExtraIndexUrls uplugin property, respectively
+	*/
+	UPROPERTY(config, EditAnywhere, Category=PythonPipInstall, AdvancedDisplay, meta=(ConfigRestartRequired=true))
+	FString ExtraInstallArgs;
 
 	/** Should remote Python execution be enabled? */
 	UPROPERTY(config, EditAnywhere, Category=PythonRemoteExecution, meta=(DisplayName="Enable Remote Execution?"))
