PipInstall: Separate pip subprocess handling code into PipRunnable. Can be extended for async notifications as well.

diff --git a/Engine/Plugins/Experimental/PythonScriptPlugin/Source/PythonScriptPlugin/Private/PipInstall.cpp b/Engine/Plugins/Experimental/PythonScriptPlugin/Source/PythonScriptPlugin/Private/PipInstall.cpp
--- a/Engine/Plugins/Experimental/PythonScriptPlugin/Source/PythonScriptPlugin/Private/PipInstall.cpp
+++ b/Engine/Plugins/Experimental/PythonScriptPlugin/Source/PythonScriptPlugin/Private/PipInstall.cpp
@@ -1,6 +1,7 @@
 // Copyright Epic Games, Inc. All Rights Reserved.
 
 #include "PipInstall.h"
+#include "PipRunnable.h"
 
 #include "PyUtil.h"
 #include "PythonScriptPluginSettings.h"
@@ -8,98 +9,14 @@
 #include "HAL/PlatformFileManager.h"
 #include "Misc/CommandLine.h"
 #include "Misc/FeedbackContext.h"
-#include "Misc/FeedbackContextMarkup.h"
 #include "Misc/FileHelper.h"
 #include "Misc/Paths.h"
 #include "Misc/ScopedSlowTask.h"
-#include "Misc/ScopeExit.h"
 
 #define LOCTEXT_NAMESPACE "PipInstall"
 
 #if WITH_PYTHON
 
-// Simple interface for parsing cmd output to update slowtask progress
-// Similar to FFeedbackContextMarkup, but supports arbitrary line parsing
-class IProgressParser
-{
-public:
-	// Get a total work estimate
-	virtual int GetTotalWork() = 0;
-	// Parse line and update status/progress (return true to eat the output and not log)
-	virtual bool UpdateStatus(const FString& ChkLine, FSlowTask& Task) = 0;
-};
-
-class FPipProgressParser : public IProgressParser
-{
-public:
-	FPipProgressParser(int GuessRequirementsCount)
-	: RequirementsDone(0)
-	, RequirementsCount(FMath::Max(GuessRequirementsCount,1.0f))
-	{}
-
-	virtual int GetTotalWork() override
-	{
-		return RequirementsCount;
-	}
-
-	virtual bool UpdateStatus(const FString& ChkLine, FSlowTask& Task) override
-	{
-		FString TrimLine = ChkLine.TrimStartAndEnd();
-		// Just log if it's not a status update line
-		if (!CheckUpdateMatch(TrimLine))
-		{
-			return false;
-		}
-
-		// Exponentially approach 100% if steps goes above estimate
-		float ProgLeft = RequirementsCount - RequirementsDone;
-		float NextWork = FMath::Clamp(0.9*ProgLeft, 0.0f, 1.0f);
-
-		// TODO: Pass in specific requirements to update status lines more accurately
-		FString StatusStr = ReplaceUpdateStrs(TrimLine);
-		Task.EnterProgressFrame(NextWork, FText::FromString(StatusStr));
-
-		RequirementsDone += NextWork;
-
-		return false;
-	}
-
-private:
-	static bool CheckUpdateMatch(const FString& Line)
-	{
-		for (const FString& ChkMatch : MatchStatusStrs)
-		{
-			if (Line.StartsWith(ChkMatch))
-			{
-				return true;
-			}
-		}
-
-		return false;
-	}
-
-	static FString ReplaceUpdateStrs(const FString& Line)
-	{
-		FString RepLine = Line;
-		for (const TPair<FString, FString>& ReplaceMap : LogReplaceStrs)
-		{
-			RepLine = RepLine.Replace(*ReplaceMap.Key, *ReplaceMap.Value, ESearchCase::CaseSensitive);
-		}
-
-		return RepLine;
-	}
-
-	float RequirementsDone;
-	float RequirementsCount;
-
-	static const TArray<FString> MatchStatusStrs;
-	static const TMap<FString,FString> LogReplaceStrs;
-};
-
-const TArray<FString> FPipProgressParser::MatchStatusStrs = {TEXT("Requirement"), TEXT("Downloading"), TEXT("Using"), TEXT("Installing")};
-const TMap<FString,FString> FPipProgressParser::LogReplaceStrs = {{TEXT("Installing collected packages:"), TEXT("Installing collected python package dependencies:")}};
-
-
 // In order to keep editor startup time fast, check directly for this utils version (make sure to match with wheel version in PythonScriptPlugin/Content/Python/Lib/wheels)
 // NOTE: This version must also be changed in PipInstallMode.cs in order to support UBT functionality
 const FString FPipInstall::PipInstallUtilsVer = TEXT("0.1.5");
@@ -110,7 +27,6 @@ const FString FPipInstall::RequirementsInputFilename = TEXT("merged_requirements
 const FString FPipInstall::ExtraUrlsFilename = TEXT("extra_urls.txt");
 const FString FPipInstall::ParsedRequirementsFilename = TEXT("merged_requirements.txt");
 
-
 FPipInstall& FPipInstall::Get()
 {
 	static FPipInstall Instance;
@@ -344,7 +260,11 @@ void FPipInstall::SetupPipEnv(FFeedbackContext* Context, bool bForceRebuild /* =
 	}
 
 	const FString VenvCmd = FString::Printf(TEXT("-m venv \"%s\""), *FPaths::ConvertRelativePathToFull(PipInstallPath));
-	int32 Res = RunPythonCmd(LOCTEXT("PipInstall.SetupVenv", "Setting up pip install environment..."), EngineInterp, VenvCmd, Context);
+
+	FScopedSlowTask SubprocTask(0.0f, LOCTEXT("PipInstall.SetupVenv", "Setting up pip install environment..."), true, *Context);
+	SubprocTask.MakeDialog();
+
+	int32 Res = RunPythonCmd(EngineInterp, VenvCmd, Context);
 	if (Res != 0)
 	{
 		UE_LOG(LogPython, Error, TEXT("Unable to create pip install environment (%d)"), Res);
@@ -373,14 +293,18 @@ FString FPipInstall::ParsePluginDependencies(const FString& MergedInRequirements
 	const FString ParsedReqsFile = PipInstallPath / ParsedRequirementsFilename;
 
 	// NOTE: Hashes are all-or-nothing so if we are ignoring, just remove them all with the parser
+	// TODO: Handle this per-plugin
 	FString DisableHashes = TEXT("");
 	if (!GetDefault<UPythonScriptPluginSettings>()->bPipStrictHashCheck)
 	{
 		DisableHashes = TEXT("--disable-hashes");
 	}
 
+	FScopedSlowTask SubprocTask(0.0f, LOCTEXT("PipInstall.ParseRequirements", "Parsing pip requirements..."), true, *Context);
+	SubprocTask.MakeDialog(false,false);
+
 	const FString Cmd = FString::Printf(TEXT("-m ue_parse_plugin_reqs %s -vv \"%s\" \"%s\""), *DisableHashes, *MergedInRequirementsFile, *ParsedReqsFile);
-	RunPythonCmd(LOCTEXT("PipInstall.ParseRequirements", "Parsing pip requirements..."), VenvInterp, Cmd, Context);
+	RunPythonCmd(VenvInterp, Cmd, Context);
 
 	return FPaths::ConvertRelativePathToFull(ParsedReqsFile);
 }
@@ -416,9 +340,11 @@ bool FPipInstall::RunPipInstall(FFeedbackContext* Context) const
 	}
 
 	const FString Cmd = SetupPipInstallCmd(ParsedReqsFile, ExtraUrls);
+	
+	TSharedRef<ICmdProgressNotifier> PipNotifier = MakeShared<FSlowTaskNotifier>(ReqCount, LOCTEXT("PipInstall.InstallRequirements", "Installing pip requirements..."), Context);
+	TSharedPtr<ICmdProgressParser> ProgParser = MakeShared<FPipProgressParser>(ReqCount, PipNotifier);
 
-	TSharedPtr<IProgressParser> ProgParser = MakeShared<FPipProgressParser>(ReqCount);
-	int32 Result = RunPythonCmd(LOCTEXT("PipInstall.InstallRequirements", "Installing pip requirements..."), VenvInterp, Cmd, Context, ProgParser);
+	int32 Result = RunPythonCmd(VenvInterp, Cmd, Context, ProgParser);
 	return (Result == 0);
 }
 
@@ -546,7 +472,7 @@ void FPipInstall::WriteSitePackagePthFile() const
 	}
 
 	// Create .pth file in PipInstall/Lib/site-packages to account for plugins with packaged dependencies
-	const FString PyPluginsSitePackageFile = GetPipSitePackagesPath() / PluginsSitePackageFilename;
+	const FString PyPluginsSitePackageFile = FPaths::ConvertRelativePathToFull(GetPipSitePackagesPath() / PluginsSitePackageFilename);
 	FFileHelper::SaveStringArrayToFile(PluginSitePackagePaths, *PyPluginsSitePackageFile);
 }
 
@@ -566,13 +492,16 @@ void FPipInstall::SetupPipInstallUtils(FFeedbackContext* Context) const
 		return;
 	}
 
-	const FString PipWheelsDir = PythonScriptDir / TEXT("Content/Python/Lib/wheels");
-	const FString InstallRequirements = PythonScriptDir / TEXT("Content/Python/PipInstallUtils/requirements.txt");
+	const FString PipWheelsDir = FPaths::ConvertRelativePathToFull(PythonScriptDir / TEXT("Content/Python/Lib/wheels"));
+	const FString InstallRequirements = FPaths::ConvertRelativePathToFull(PythonScriptDir / TEXT("Content/Python/PipInstallUtils/requirements.txt"));
 
 	const FString PipInstallReq = TEXT("ue-pipinstall-utils==") + PipInstallUtilsVer;
 	const FString Cmd = FString::Printf(TEXT("-m pip install --upgrade --no-index --find-links \"%s\" -r \"%s\" %s"), *PipWheelsDir, *InstallRequirements, *PipInstallReq);
 
-	RunPythonCmd(LOCTEXT("PipInstall.SetupPipInstallUtils", "Setting up pip install utils"), VenvInterp, Cmd, Context);
+	FScopedSlowTask SubprocTask(0.0f, LOCTEXT("PipInstall.SetupPipInstallUtils", "Setting up pip install utils"), true, *Context);
+	SubprocTask.MakeDialog(false,false);
+
+	RunPythonCmd(VenvInterp, Cmd, Context);
 }
 
 
@@ -580,7 +509,10 @@ bool FPipInstall::CheckPipInstallUtils(FFeedbackContext* Context) const
 {
 	// Verify that correct version of pip install utils is already available
 	const FString Cmd = FString::Printf(TEXT("-c \"import pkg_resources;dist=pkg_resources.working_set.find(pkg_resources.Requirement.parse('ue-pipinstall-utils'));exit(dist.version!='%s' if dist is not None else 1)\""), *PipInstallUtilsVer);
-	return (RunPythonCmd(LOCTEXT("PipInstall.CheckPipInstallUtils", "Check pip install utils installed"), VenvInterp, Cmd, Context) == 0);
+	FScopedSlowTask SubprocTask(0.0f, LOCTEXT("PipInstall.CheckPipInstallUtils", "Check pip install utils installed"), true, *Context);
+	SubprocTask.MakeDialog(false,false);
+
+	return (RunPythonCmd(VenvInterp, Cmd, Context) == 0);
 }
 
 FString FPipInstall::SetupPipInstallCmd(const FString& ParsedReqsFile, const TArray<FString>& ExtraUrls) const
@@ -621,75 +553,22 @@ FString FPipInstall::SetupPipInstallCmd(const FString& ParsedReqsFile, const TAr
 	return Cmd;
 }
 
-int32 FPipInstall::RunPythonCmd(const FText& Description, const FString& PythonInterp, const FString& Cmd, FFeedbackContext* Context, TSharedPtr<IProgressParser> CmdParser)
+int32 FPipInstall::RunPythonCmd(const FString& PythonInterp, const FString& Cmd, FFeedbackContext* Context, TSharedPtr<ICmdProgressParser> CmdParser)
 {
 	TRACE_CPUPROFILER_EVENT_SCOPE(FPipInstall::RunPythonCmd);
 
-	UE_LOG(LogPython, Log, TEXT("Running python command: \"%s\" %s"), *PythonInterp, *Cmd);
-
-	int32 Result = 0;
-	RunLoggedSubprocess(&Result, Description, PythonInterp, Cmd, Context, CmdParser);
-
-	return Result;
-}
-
-bool FPipInstall::RunLoggedSubprocess(int32* OutExitCode, const FText& Description, const FString& URL, const FString& Params, FFeedbackContext* Context, TSharedPtr<IProgressParser> CmdParser)
-{
-	int AmountOfWork = (CmdParser.IsValid()) ? CmdParser->GetTotalWork() : 0;
-
-	FScopedSlowTask SubprocessTask(AmountOfWork, Description, true, *Context);
-	SubprocessTask.Visibility = ESlowTaskVisibility::Important; // this function can be very slow, users will benefit from our messages
-	SubprocessTask.MakeDialog();
+	UE_LOG(LogPython, Log, TEXT("Running python command: python %s"), *Cmd);
 
-	// Create a read and write pipe for the child process
-	void* StdOutPipeRead = nullptr;
-	void* StdOutPipeWrite = nullptr;
-	verify(FPlatformProcess::CreatePipe(StdOutPipeRead, StdOutPipeWrite));
-
-	ON_SCOPE_EXIT
-	{
-		FPlatformProcess::ClosePipe(StdOutPipeRead, StdOutPipeWrite);
-	};
-
-	// Create the process
-	FProcHandle ProcessHandle = FPlatformProcess::CreateProc(*URL, *Params, false, true, true, nullptr, 0, nullptr, StdOutPipeWrite, nullptr);
-	if (ProcessHandle.IsValid())
-	{
-		FString BufferedText;
-		for (bool bProcessFinished = false; !bProcessFinished; )
-		{
-			bProcessFinished = FPlatformProcess::GetProcReturnCode(ProcessHandle, OutExitCode);
-			BufferedText += FPlatformProcess::ReadPipe(StdOutPipeRead);
-
-			int32 EndOfLineIdx;
-			while (BufferedText.FindChar(TEXT('\n'), EndOfLineIdx))
-			{
-				FString Line = BufferedText.Left(EndOfLineIdx);
-				Line.RemoveFromEnd(TEXT("\r"), ESearchCase::CaseSensitive);
-
-				// Always log if no output parser, also log if UpdateStatus returns false
-				if (!CmdParser.IsValid() || !CmdParser->UpdateStatus(Line, SubprocessTask))
-				{
-					Context->Log(LogPython.GetCategoryName(), ELogVerbosity::Log, Line);
-				}
-
-				BufferedText.MidInline(EndOfLineIdx + 1, MAX_int32, EAllowShrinking::No);
-			}
-
-			FPlatformProcess::Sleep(0.1f);
-		}
-		ProcessHandle.Reset();
-		return true;
-	}
-	else
+	int32 OutResult = 0;
+	if (!FLoggedSubprocessSync::Run(OutResult, FPaths::ConvertRelativePathToFull(PythonInterp), Cmd, Context, CmdParser))
 	{
-		Context->CategorizedLogf(LogPython.GetCategoryName(), ELogVerbosity::Warning, TEXT("Couldn't create process '%s'"), *URL);
+		UE_LOG(LogPython, Error, TEXT("Unable to create python process"));
+		return -1;
 	}
 
-	return false;
+	return OutResult;
 }
 
-
 FString FPipInstall::GetPythonScriptPluginPath()
 {
 	TSharedPtr<IPlugin> PythonPlugin = IPluginManager::Get().FindPlugin("PythonScriptPlugin");
diff --git a/Engine/Plugins/Experimental/PythonScriptPlugin/Source/PythonScriptPlugin/Private/PipInstall.h b/Engine/Plugins/Experimental/PythonScriptPlugin/Source/PythonScriptPlugin/Private/PipInstall.h
--- a/Engine/Plugins/Experimental/PythonScriptPlugin/Source/PythonScriptPlugin/Private/PipInstall.h
+++ b/Engine/Plugins/Experimental/PythonScriptPlugin/Source/PythonScriptPlugin/Private/PipInstall.h
@@ -3,6 +3,7 @@
 #pragma once
 
 #include "Containers/Array.h"
+#include "Containers/UnrealString.h"
 #include "Interfaces/IPluginManager.h"
 #include "Templates/SharedPointer.h"
 
@@ -11,7 +12,7 @@
 class FJsonObject;
 class FFeedbackContext;
 
-class IProgressParser;
+struct ICmdProgressParser;
 
 class FPipInstall
 {
@@ -46,8 +47,7 @@ private:
 
 	FString SetupPipInstallCmd(const FString& ParsedReqsFile, const TArray<FString>& ExtraUrls) const;
 
-	static int32 RunPythonCmd(const FText& Description, const FString& PythonInterp, const FString& Cmd, FFeedbackContext* Context, TSharedPtr<IProgressParser> CmdParser = nullptr);
-	static bool RunLoggedSubprocess(int32* OutExitCode, const FText& Description, const FString& URL, const FString& Params, FFeedbackContext* Context, TSharedPtr<IProgressParser> CmdParser);
+	static int32 RunPythonCmd(const FString& PythonInterp, const FString& Cmd, FFeedbackContext* Context, TSharedPtr<ICmdProgressParser> CmdParser = nullptr);
 
 	FString ParseVenvVersion() const;
 
diff --git a/Engine/Plugins/Experimental/PythonScriptPlugin/Source/PythonScriptPlugin/Private/PipRunnable.cpp b/Engine/Plugins/Experimental/PythonScriptPlugin/Source/PythonScriptPlugin/Private/PipRunnable.cpp
new file mode 100644
--- /dev/null
+++ b/Engine/Plugins/Experimental/PythonScriptPlugin/Source/PythonScriptPlugin/Private/PipRunnable.cpp
@@ -0,0 +1,204 @@
+// Copyright Epic Games, Inc. All Rights Reserved.
+
+#include "PipRunnable.h"
+
+#include "PyUtil.h"
+
+#include "Misc/AsyncTaskNotification.h"
+#include "Misc/FeedbackContext.h"
+#include "Misc/SlowTask.h"
+
+#define LOCTEXT_NAMESPACE "PipInstall"
+
+#if WITH_PYTHON
+
+// FPipProgressParser Implementation
+FPipProgressParser::FPipProgressParser(int GuessRequirementsCount, TSharedRef<ICmdProgressNotifier> InCmdNotifier)
+	: RequirementsDone(0.0f)
+	, RequirementsCount(FMath::Max(2.0f*GuessRequirementsCount,1.0f))
+	, CmdNotifier(InCmdNotifier)
+{}
+
+float FPipProgressParser::GetTotalWork()
+{
+	return RequirementsCount;
+}
+
+float FPipProgressParser::GetWorkDone()
+{
+	return RequirementsDone;
+}
+
+bool FPipProgressParser::UpdateStatus(const FString& ChkLine)
+{
+	FString TrimLine = ChkLine.TrimStartAndEnd();
+	// Just log if it's not a status update line
+	if (!CheckUpdateMatch(TrimLine))
+	{
+		return false;
+	}
+
+	FText Status = FText::FromString(ReplaceUpdateStrs(TrimLine));
+	CmdNotifier->UpdateProgress(RequirementsDone, RequirementsCount, Status);
+
+	RequirementsDone += 1;
+	RequirementsCount = FMath::Max(RequirementsCount, RequirementsDone + 1);
+
+	// Note: Return true instead to not log status matched lines
+	return false;
+}
+
+void FPipProgressParser::NotifyCompleted(bool bSuccess)
+{
+	CmdNotifier->Completed(bSuccess);
+}
+
+
+bool FPipProgressParser::CheckUpdateMatch(const FString& Line)
+{
+	for (const FString& ChkMatch : MatchStatusStrs)
+	{
+		if (Line.StartsWith(ChkMatch))
+		{
+			return true;
+		}
+	}
+
+	return false;
+}
+
+FString FPipProgressParser::ReplaceUpdateStrs(const FString& Line)
+{
+	FString RepLine = Line;
+	for (const TPair<FString, FString>& ReplaceMap : LogReplaceStrs)
+	{
+		RepLine = RepLine.Replace(*ReplaceMap.Key, *ReplaceMap.Value, ESearchCase::CaseSensitive);
+	}
+
+	return RepLine;
+}
+
+
+const TArray<FString> FPipProgressParser::MatchStatusStrs = {TEXT("Requirement"), TEXT("Downloading"), TEXT("Collecting"), TEXT("Using"), TEXT("Installing")};
+const TMap<FString,FString> FPipProgressParser::LogReplaceStrs = {{TEXT("Installing collected packages:"), TEXT("Installing collected python package dependencies:")}};
+
+
+// FSlowTaskNotifier implementation
+FSlowTaskNotifier::FSlowTaskNotifier(float GuessSteps, const FText& Description, FFeedbackContext* Context)
+	: SlowTask(nullptr)
+	, TotalWork(FMath::Max(GuessSteps, 1.0f))
+	, WorkDone(0.0f)
+{
+	SlowTask = MakeUnique<FSlowTask>(GuessSteps, Description, true, *Context);
+	WorkDone = SlowTask->CompletedWork;
+}
+
+void FSlowTaskNotifier::UpdateProgress(float UpdateWorkDone, float UpdateTotalWork, const FText& Status)
+{
+	// Exponentially approach 100% if work goes above estimate
+	float NextDone = TotalWork * UpdateWorkDone / UpdateTotalWork;
+	float NextStep = FMath::Max(NextDone - WorkDone, 0.0f);
+	float ProgLeft = FMath::Max(TotalWork - WorkDone, 0.0f);
+
+	// Clamp to 90% of remaining progress
+	float NextWork = FMath::Clamp(NextStep, 0.0f, 0.9*ProgLeft);
+
+	// TODO: Pass in specific requirements to update status lines more accurately
+	SlowTask->EnterProgressFrame(NextWork, Status);
+	WorkDone += NextWork;
+}
+
+void FSlowTaskNotifier::Completed(bool bSuccess){}
+
+
+// Implementation of a tickable class for running a subprocess with status updates
+class FRunLoggedSubprocess
+{
+public:
+	FRunLoggedSubprocess(const FString& URI, const FString& Params, FFeedbackContext* Context, TSharedPtr<ICmdProgressParser> CmdParser)
+		: Context(Context), CmdParser(CmdParser)
+	{
+		bLogOutput = (Context != nullptr);
+
+		// Create a stdout pipe for the child process
+		StdOutPipeRead = nullptr;
+		StdOutPipeWrite = nullptr;
+		verify(FPlatformProcess::CreatePipe(StdOutPipeRead, StdOutPipeWrite));
+
+		ProcExitCode = -1;
+		ProcessHandle = FPlatformProcess::CreateProc(*URI, *Params, false, true, true, nullptr, 0, nullptr, StdOutPipeWrite, nullptr);
+	}
+
+	~FRunLoggedSubprocess()
+	{
+		FPlatformProcess::ClosePipe(StdOutPipeRead, StdOutPipeWrite);
+	}
+
+	bool IsValid()
+	{
+		return ProcessHandle.IsValid();
+	}
+
+	int32 GetExitCode()
+	{
+		return ProcExitCode;
+	}
+
+	bool TickSubprocess()
+	{
+		bool bProcessFinished = FPlatformProcess::GetProcReturnCode(ProcessHandle, &ProcExitCode);
+		BufferedText += FPlatformProcess::ReadPipe(StdOutPipeRead);
+
+		int32 EndOfLineIdx;
+		while (BufferedText.FindChar(TEXT('\n'), EndOfLineIdx))
+		{
+			FString Line = BufferedText.Left(EndOfLineIdx);
+			Line.RemoveFromEnd(TEXT("\r"), ESearchCase::CaseSensitive);
+
+			// Always log if no output parser, also log if UpdateStatus returns false
+			bool bOnlyStatusUpdate = CmdParser.IsValid() && CmdParser->UpdateStatus(Line);
+			if ( bLogOutput && !bOnlyStatusUpdate)
+			{
+				Context->Log(LogPython.GetCategoryName(), ELogVerbosity::Log, Line);
+			}
+
+			BufferedText.MidInline(EndOfLineIdx + 1, MAX_int32, EAllowShrinking::No);
+		}
+
+		return !bProcessFinished;
+	}
+
+private:
+	FFeedbackContext* Context;
+	TSharedPtr<ICmdProgressParser> CmdParser;
+
+	bool bLogOutput;
+	int32 ProcExitCode;
+	FProcHandle ProcessHandle;
+	void* StdOutPipeRead = nullptr;
+	void* StdOutPipeWrite = nullptr;
+
+	FString BufferedText;
+};
+
+
+bool FLoggedSubprocessSync::Run(int32& OutExitCode, const FString& URI, const FString& Params, FFeedbackContext* Context, TSharedPtr<ICmdProgressParser> CmdParser)
+{
+	FRunLoggedSubprocess Subproc(URI, Params, Context, CmdParser);
+	if (!Subproc.IsValid())
+	{
+		return false;
+	}
+
+	while (Subproc.TickSubprocess())
+	{
+		FPlatformProcess::Sleep(0.1f);
+	}
+
+	OutExitCode = Subproc.GetExitCode();
+	return true;
+}
+
+#endif //WITH_PYTHON
+
+#undef LOCTEXT_NAMESPACE
diff --git a/Engine/Plugins/Experimental/PythonScriptPlugin/Source/PythonScriptPlugin/Private/PipRunnable.h b/Engine/Plugins/Experimental/PythonScriptPlugin/Source/PythonScriptPlugin/Private/PipRunnable.h
new file mode 100644
--- /dev/null
+++ b/Engine/Plugins/Experimental/PythonScriptPlugin/Source/PythonScriptPlugin/Private/PipRunnable.h
@@ -0,0 +1,81 @@
+// Copyright Epic Games, Inc. All Rights Reserved.
+
+#pragma once
+
+#if WITH_PYTHON
+
+#include "Containers/UnrealString.h"
+#include "HAL/Runnable.h"
+#include "HAL/RunnableThread.h"
+#include "Misc/AsyncTaskNotification.h"
+#include "Templates/SharedPointer.h"
+
+struct FSlowTask;
+
+// Simple interface for parsing cmd output to update slowtask progress
+// Similar to FFeedbackContextMarkup, but supports arbitrary line parsing
+struct ICmdProgressParser
+{
+    virtual ~ICmdProgressParser(){};
+	// Get a total work estimate
+	virtual float GetTotalWork() = 0;
+    virtual float GetWorkDone() = 0;
+	// Parse line and update status/progress (return true to eat the output and not log)
+	virtual bool UpdateStatus(const FString& ChkLine) = 0;
+	virtual void NotifyCompleted(bool bSuccess) = 0;
+};
+
+// Interface to wrap notifications for progress updates (e.g. slowtask or asyncnotify)
+struct ICmdProgressNotifier
+{
+    virtual ~ICmdProgressNotifier(){};
+	virtual void UpdateProgress(float UpdateWorkDone, float UpdateTotalWork, const FText& Status) = 0;
+	virtual void Completed(bool bSuccess) = 0;
+};
+
+// Pip progress parser implemenation of ICmdProgressParser
+class FPipProgressParser : public ICmdProgressParser
+{
+public:
+	FPipProgressParser(int GuessRequirementsCount, TSharedRef<ICmdProgressNotifier> InCmdNotifier);
+    // ICmdProgressParser methods
+	virtual float GetTotalWork() override;
+    virtual float GetWorkDone() override;
+	virtual bool UpdateStatus(const FString& ChkLine) override;
+	virtual void NotifyCompleted(bool bSuccess) override;
+
+private:
+	static bool CheckUpdateMatch(const FString& Line);
+	static FString ReplaceUpdateStrs(const FString& Line);
+
+	float RequirementsDone;
+	float RequirementsCount;
+    TSharedRef<ICmdProgressNotifier> CmdNotifier;
+
+	static const TArray<FString> MatchStatusStrs;
+	static const TMap<FString,FString> LogReplaceStrs;
+};
+
+// Slowtask-based notifier for updating command progress (ICmdProgressNotifier)
+class FSlowTaskNotifier : public ICmdProgressNotifier
+{
+public:
+	FSlowTaskNotifier(float GuessSteps, const FText& Description, FFeedbackContext* Context);
+    // ICmdProgressNotifier methods
+	virtual void UpdateProgress(float UpdateWorkDone, float UpdateTotalWork, const FText& Status) override;
+	virtual void Completed(bool bSuccess) override;
+
+private:
+	TUniquePtr<FSlowTask> SlowTask;
+
+	float TotalWork;
+	float WorkDone;
+};
+
+// Run a subprocess synchronously (assumes running in game thread)
+struct FLoggedSubprocessSync
+{
+    static bool Run(int32& OutExitCode, const FString& URI, const FString& Params, FFeedbackContext* Context, TSharedPtr<ICmdProgressParser> CmdParser);
+};
+
+#endif //WITH_PYTHON
