Added more control over whether Python is enabled by default, when included in the build

The CVar 'Engine.Python.IsEnabledByDefault' can now be set (prior to OnPostEngineInit) to control whether Python will be enabled by default at runtime.

This CVar is lower priority than the existing -DisablePython and DisablePythonForCommandlet options that existed to disable Python at runtime.

Additionally, two extra command line arguments are now available to override whether Python is enabled at runtime:
  * -ForceEnablePython will force Python to be enabled at runtime. It is higher precedence than any other method of disabling Python.
  * -EnablePython will enable Python at runtime, even if 'Engine.Python.IsEnabledByDefault' is false. It is lower precedence than the command line and config-based methods of disabling Python.

Note: UE_PYTHON_DEFER_INIT has been removed as part of this change. It has been the default for a while now, and deferred initialization is required to allow things to override 'Engine.Python.IsEnabledByDefault' prior to OnPostEngineInit.

diff --git a/Engine/Plugins/Experimental/PythonScriptPlugin/Source/PythonScriptPlugin/Private/PythonScriptPlugin.cpp b/Engine/Plugins/Experimental/PythonScriptPlugin/Source/PythonScriptPlugin/Private/PythonScriptPlugin.cpp
--- a/Engine/Plugins/Experimental/PythonScriptPlugin/Source/PythonScriptPlugin/Private/PythonScriptPlugin.cpp
+++ b/Engine/Plugins/Experimental/PythonScriptPlugin/Source/PythonScriptPlugin/Private/PythonScriptPlugin.cpp
@@ -69,14 +69,15 @@
 
 #define LOCTEXT_NAMESPACE "PythonScriptPlugin"
 
-#define UE_PYTHON_DEFER_INIT (1)
-
 #if WITH_PYTHON
 
 #if PY_MAJOR_VERSION >= 3 && PY_MINOR_VERSION >= 11
 	PyUtil::FPyApiBuffer FPythonScriptPlugin::Utf8String= PyUtil::TCHARToPyApiBuffer(TEXT("utf-8"));
 #endif
 
+namespace UE::Python::Private
+{
+
 static PyUtil::FPyApiBuffer NullPyArg = PyUtil::TCHARToPyApiBuffer(TEXT(""));
 static PyUtil::FPyApiChar* NullPyArgPtrs[] = { NullPyArg.GetData() };
 
@@ -89,6 +90,71 @@ FPyObjectPtr MakeEmptyArgvList()
 	return PyArgvList;
 }
 
+static bool bIsEnabledByDefault = true;
+static FAutoConsoleVariableRef CVarAsyncLoadLocalizationData(TEXT("Engine.Python.IsEnabledByDefault"), bIsEnabledByDefault, TEXT("True if Python is enabled by default (after checking any command line overrides), or False if it should be considered disabled. This must be set prior to OnPostEngineInit."));
+
+bool IsPythonEnabled()
+{
+	if (FParse::Param(FCommandLine::Get(), TEXT("ForceEnablePython")))
+	{
+		UE_LOG(LogPython, Log, TEXT("Python enabled via command-line flag '-ForceEnablePython'"));
+		return true;
+	}
+
+	if (FParse::Param(FCommandLine::Get(), TEXT("DisablePython")))
+	{
+		UE_LOG(LogPython, Log, TEXT("Python disabled via command-line flag '-DisablePython'"));
+		return false;
+	}
+
+	if (IsRunningCommandlet())
+	{
+		TArray<FString> DisablePythonForCommandlets;
+		GConfig->GetArray(TEXT("PythonScriptPlugin"), TEXT("DisablePythonForCommandlet"), DisablePythonForCommandlets, GEditorIni);
+
+		FString RunningCommandletName;
+		if (!DisablePythonForCommandlets.IsEmpty() && FParse::Value(FCommandLine::Get(), TEXT("-run="), RunningCommandletName))
+		{
+			auto CleanCommandletName = [](FString& InOutCommandletName)
+			{
+				const FStringView CommandletSuffix = TEXTVIEW("Commandlet");
+				if (InOutCommandletName.EndsWith(CommandletSuffix))
+				{
+					InOutCommandletName.LeftChopInline(CommandletSuffix.Len(), EAllowShrinking::No);
+				}
+			};
+
+			CleanCommandletName(RunningCommandletName);
+			for (FString& DisablePythonForCommandlet : DisablePythonForCommandlets)
+			{
+				CleanCommandletName(DisablePythonForCommandlet);
+				if (DisablePythonForCommandlet == RunningCommandletName)
+				{
+					UE_LOG(LogPython, Log, TEXT("Python disabled via config setting 'DisablePythonForCommandlet'"));
+					return false;
+				}
+			}
+		}
+	}
+
+	if (FParse::Param(FCommandLine::Get(), TEXT("EnablePython")))
+	{
+		UE_LOG(LogPython, Log, TEXT("Python enabled via command-line flag '-EnablePython'"));
+		return true;
+	}
+
+	if (UE::Python::Private::bIsEnabledByDefault)
+	{
+		UE_LOG(LogPython, Log, TEXT("Python enabled via CVar 'Engine.Python.IsEnabledByDefault'"));
+		return true;
+	}
+
+	UE_LOG(LogPython, Log, TEXT("Python disabled via CVar 'Engine.Python.IsEnabledByDefault'"));
+	return false;
+}
+
+} // namespace UE::Python::Private
+
 /** Util struct to set the sys.argv data for Python when executing a file with arguments */
 struct FPythonScopedArgv
 {
@@ -130,10 +196,10 @@ struct FPythonScopedArgv
 	~FPythonScopedArgv()
 	{
 #if PY_MAJOR_VERSION >= 3 && PY_MINOR_VERSION >= 11
-		FPyObjectPtr PyArgvList = MakeEmptyArgvList();
+		FPyObjectPtr PyArgvList = UE::Python::Private::MakeEmptyArgvList();
 		PySys_SetObject("argv", PyArgvList.Get());
 #else
-		PySys_SetArgvEx(1, NullPyArgPtrs, 0);
+		PySys_SetArgvEx(1, UE::Python::Private::NullPyArgPtrs, 0);
 #endif
 	}
 
@@ -589,9 +655,6 @@ FPythonScriptPlugin::FPythonScriptPlugin()
 #if WITH_PYTHON
 	: CmdExec(this)
 	, CmdREPLExec(this)
-	, CmdMenu(nullptr)
-	, bInitialized(false)
-	, bRanStartupScripts(false)
 #endif	// WITH_PYTHON
 {
 }
@@ -599,48 +662,7 @@ FPythonScriptPlugin::FPythonScriptPlugin()
 bool FPythonScriptPlugin::IsPythonAvailable() const
 {
 #if WITH_PYTHON
-	auto IsPythonEnabled = []()
-	{
-		if (FParse::Param(FCommandLine::Get(), TEXT("DisablePython")))
-		{
-			UE_LOG(LogPython, Log, TEXT("Python disabled via command-line flag '-DisablePython'"));
-			return false;
-		}
-
-		if (IsRunningCommandlet())
-		{
-			TArray<FString> DisablePythonForCommandlets;
-			GConfig->GetArray(TEXT("PythonScriptPlugin"), TEXT("DisablePythonForCommandlet"), DisablePythonForCommandlets, GEditorIni);
-
-			FString RunningCommandletName;
-			if (!DisablePythonForCommandlets.IsEmpty() && FParse::Value(FCommandLine::Get(), TEXT("-run="), RunningCommandletName))
-			{
-				auto CleanCommandletName = [](FString& InOutCommandletName)
-				{
-					const FStringView CommandletSuffix = TEXTVIEW("Commandlet");
-					if (InOutCommandletName.EndsWith(CommandletSuffix))
-					{
-						InOutCommandletName.LeftChopInline(CommandletSuffix.Len(), EAllowShrinking::No);
-					}
-				};
-
-				CleanCommandletName(RunningCommandletName);
-				for (FString& DisablePythonForCommandlet : DisablePythonForCommandlets)
-				{
-					CleanCommandletName(DisablePythonForCommandlet);
-					if (DisablePythonForCommandlet == RunningCommandletName)
-					{
-						UE_LOG(LogPython, Log, TEXT("Python disabled via config setting 'DisablePythonForCommandlet'"));
-						return false;
-					}
-				}
-			}
-		}
-
-		return true;
-	};
-	static const bool bEnablePython = IsPythonEnabled();
-	return bEnablePython;
+	return bIsEnabled;
 #else
 	return false;
 #endif
@@ -785,17 +807,28 @@ void FPythonScriptPlugin::StartupModule()
 	FEditorPythonExecuter::OnStartupModule();
 #endif
 
+#if WITH_PYTHON
+	LLM_SCOPE_BYNAME(TEXT("PythonScriptPlugin"));
+
+	FCoreDelegates::OnPostEngineInit.AddRaw(this, &FPythonScriptPlugin::OnPostEngineInit);
+#endif	// WITH_PYTHON
+}
+
+#if WITH_PYTHON
+void FPythonScriptPlugin::OnPostEngineInit()
+{
+	LLM_SCOPE_BYNAME(TEXT("PythonScriptPlugin"));
+
+	// Determine whether Python is enabled in OnPostEngineInit so that other things have time to override the CVar
+	// IsPythonAvailable() will always return false prior to configuring bIsEnabled
+	bIsEnabled = UE::Python::Private::IsPythonEnabled();
+	
 	if (!IsPythonAvailable())
 	{
 		return;
 	}
 
-#if WITH_PYTHON
-	LLM_SCOPE_BYNAME(TEXT("PythonScriptPlugin"));
-
-#if !UE_PYTHON_DEFER_INIT
 	InitializePython();
-#endif
 
 	IModularFeatures::Get().RegisterModularFeature(IConsoleCommandExecutor::ModularFeatureName(), &CmdExec);
 	IModularFeatures::Get().RegisterModularFeature(IConsoleCommandExecutor::ModularFeatureName(), &CmdREPLExec);
@@ -803,19 +836,7 @@ void FPythonScriptPlugin::StartupModule()
 	check(!RemoteExecution);
 	RemoteExecution = MakeUnique<FPythonScriptRemoteExecution>(this);
 
-	FCoreDelegates::OnPostEngineInit.AddRaw(this, &FPythonScriptPlugin::OnPostEngineInit);
 	FCoreDelegates::OnPreExit.AddRaw(this, &FPythonScriptPlugin::ShutdownPython);
-#endif	// WITH_PYTHON
-}
-
-#if WITH_PYTHON
-void FPythonScriptPlugin::OnPostEngineInit()
-{
-	LLM_SCOPE_BYNAME(TEXT("PythonScriptPlugin"));
-
-#if UE_PYTHON_DEFER_INIT
-	InitializePython();
-#endif
 
 #if WITH_EDITOR
 	FPythonScriptPluginEditorStyle::Get();
@@ -840,11 +861,6 @@ void FPythonScriptPlugin::ShutdownModule()
 	FEditorPythonExecuter::OnShutdownModule();
 #endif
 
-	if (!IsPythonAvailable())
-	{
-		return;
-	}
-
 #if WITH_PYTHON
 	FCoreDelegates::OnPreExit.RemoveAll(this);
 	FCoreDelegates::OnPostEngineInit.RemoveAll(this);
@@ -994,10 +1010,10 @@ void FPythonScriptPlugin::InitializePython()
 	{
 #if PY_MAJOR_VERSION >=3 && PY_MINOR_VERSION >= 11
 		// Set default argv to [""]
-		FPyObjectPtr PyArgvList = MakeEmptyArgvList();
+		FPyObjectPtr PyArgvList = UE::Python::Private::MakeEmptyArgvList();
 		PySys_SetObject("argv", PyArgvList.Get());
 #else
-		PySys_SetArgvEx(1, NullPyArgPtrs, 0);
+		PySys_SetArgvEx(1, UE::Python::Private::NullPyArgPtrs, 0);
 #endif
 
 		// Enable developer warnings if requested
@@ -1016,6 +1032,11 @@ void FPythonScriptPlugin::InitializePython()
 		InitializePyConstant();
 
 		PyObject* PyMainModule = PyImport_AddModule("__main__");
+		if (!PyMainModule)
+		{
+			PyUtil::LogPythonError();
+			checkf(false, TEXT("Failed to initialize '__main__'! This typically means the Python SDK is missing or could not be loaded"));
+		}
 		PyDefaultGlobalDict = FPyObjectPtr::NewReference(PyModule_GetDict(PyMainModule));
 		PyDefaultLocalDict = PyDefaultGlobalDict;
 
@@ -1238,6 +1259,7 @@ void FPythonScriptPlugin::ShutdownPython()
 	Py_Finalize();
 
 	bInitialized = false;
+	bIsEnabled = false;
 	bRanStartupScripts = false;
 }
 
@@ -1932,8 +1954,6 @@ void FPythonScriptPlugin::PopulatePythonFileContextMenu(UToolMenu* InMenu)
 
 #endif	// WITH_PYTHON
 
-#undef UE_PYTHON_DEFER_INIT
-
 IMPLEMENT_MODULE(FPythonScriptPlugin, PythonScriptPlugin)
 
 #undef LOCTEXT_NAMESPACE
diff --git a/Engine/Plugins/Experimental/PythonScriptPlugin/Source/PythonScriptPlugin/Private/PythonScriptPlugin.h b/Engine/Plugins/Experimental/PythonScriptPlugin/Source/PythonScriptPlugin/Private/PythonScriptPlugin.h
--- a/Engine/Plugins/Experimental/PythonScriptPlugin/Source/PythonScriptPlugin/Private/PythonScriptPlugin.h
+++ b/Engine/Plugins/Experimental/PythonScriptPlugin/Source/PythonScriptPlugin/Private/PythonScriptPlugin.h
@@ -205,7 +205,7 @@ private:
 	TUniquePtr<FPythonScriptRemoteExecution> RemoteExecution;
 	FPythonCommandExecutor CmdExec;
 	FPythonREPLCommandExecutor CmdREPLExec;
-	IPythonCommandMenu* CmdMenu;
+	IPythonCommandMenu* CmdMenu = nullptr;
 	FTSTicker::FDelegateHandle TickHandle;
 	FTSTicker::FDelegateHandle ModuleDelayedHandle;
 
@@ -217,8 +217,9 @@ private:
 	FPyObjectPtr PyConsoleLocalDict;
 	FPyObjectPtr PyUnrealModule;
 	PyThreadState* PyMainThreadState = nullptr;
-	bool bInitialized;
-	bool bRanStartupScripts;
+	bool bInitialized = false;
+	bool bIsEnabled = false;
+	bool bRanStartupScripts = false;
 
 #if PY_MAJOR_VERSION >= 3 && PY_MINOR_VERSION >= 11
 	PyConfig ModulePyConfig;
