Added the missing CPF_OutParm flag to PythonGeneratedClass function return parameters

This flag prevents UObject::ProcessEvent overwriting the Parms data (which was filled in by the function call) with the NewStack.Locals data (which is empty), as it thought the return property was a local variable due to it only checking CPF_OutParm. Other CPF_ReturnParm properties are also marked CPF_OutParm, so this makes Python consistent.

#rnx

diff --git a/Engine/Plugins/Experimental/PythonScriptPlugin/Source/PythonScriptPlugin/Private/PyWrapperObject.cpp b/Engine/Plugins/Experimental/PythonScriptPlugin/Source/PythonScriptPlugin/Private/PyWrapperObject.cpp
--- a/Engine/Plugins/Experimental/PythonScriptPlugin/Source/PythonScriptPlugin/Private/PyWrapperObject.cpp
+++ b/Engine/Plugins/Experimental/PythonScriptPlugin/Source/PythonScriptPlugin/Private/PyWrapperObject.cpp
@@ -1684,7 +1684,7 @@ public:
 					PyUtil::SetPythonError(PyExc_Exception, PyType, *FString::Printf(TEXT("Failed to create return property (%s) for function '%s'"), *PyUtil::GetFriendlyTypename(RetType), *InFieldName));
 					return false;
 				}
-				RetProp->PropertyFlags |= (CPF_Parm | CPF_ReturnParm);
+				RetProp->PropertyFlags |= (CPF_Parm | CPF_OutParm | CPF_ReturnParm);
 				Func->AddCppProperty(RetProp);
 
 				if (bOptionalReturn)
