Python stub fixes

* Fixed FName/FText default values not being escaped for the script literal
* Fixed deprecations being added with invalid symbols due to redirecting from a BP
* Fixed unreal_core.py using the wrong imported name of the io module

#rnx

diff --git a/Engine/Plugins/Experimental/PythonScriptPlugin/Source/PythonScriptPlugin/Private/PyGenUtil.cpp b/Engine/Plugins/Experimental/PythonScriptPlugin/Source/PythonScriptPlugin/Private/PyGenUtil.cpp
--- a/Engine/Plugins/Experimental/PythonScriptPlugin/Source/PythonScriptPlugin/Private/PyGenUtil.cpp
+++ b/Engine/Plugins/Experimental/PythonScriptPlugin/Source/PythonScriptPlugin/Private/PyGenUtil.cpp
@@ -2243,18 +2243,19 @@ void PythonizeValueImpl(const FProperty* InProp, const void* InPropValue, const
 		}
 		else if (const FNameProperty* NameProp = CastField<const FNameProperty>(InProp))
 		{
-			const FString NameStrValue = NameProp->GetPropertyValue(PropArrValue).ToString();
+			FString NameStrValue = NameProp->GetPropertyValue(PropArrValue).ToString();
+			NameStrValue.ReplaceCharWithEscapedCharInline();
 			OutPythonDefaultValue += bUseStrictTyping
 				? FString::Printf(TEXT("%sName(\"%s\")"), UnrealNamespace, *NameStrValue)
 				: FString::Printf(TEXT("\"%s\""), *NameStrValue);
 		}
 		else if (const FTextProperty* TextProp = CastField<const FTextProperty>(InProp))
 		{
-			const FString* TextStrValue = FTextInspector::GetSourceString(TextProp->GetPropertyValue(PropArrValue));
-			check(TextStrValue);
+			FString TextStrValue = TextProp->GetPropertyValue(PropArrValue).BuildSourceString();
+			TextStrValue.ReplaceCharWithEscapedCharInline();
 			OutPythonDefaultValue += bUseStrictTyping
-				? FString::Printf(TEXT("%sText(\"%s\")"), UnrealNamespace, **TextStrValue)
-				: FString::Printf(TEXT("\"%s\""), **TextStrValue);
+				? FString::Printf(TEXT("%sText(\"%s\")"), UnrealNamespace, *TextStrValue)
+				: FString::Printf(TEXT("\"%s\""), *TextStrValue);
 		}
 		else if (const FFieldPathProperty* FieldPathProp =  CastField<const FFieldPathProperty>(InProp))
 		{
@@ -2848,6 +2849,13 @@ TArray<TTuple<FSoftObjectPath, FString>> GetDeprecatedFieldPythonNamesImpl(const
 			continue;
 		}
 		
+		// If the previous name contains invalid characters, then skip it
+		// It is likely redirecting from a BP, which has weaker naming restrictions than C++ or Python
+		if (!IsValidName(FNameBuilder(PreviousName.ObjectName).ToView()))
+		{
+			continue;
+		}
+
 		// Class/Struct redirects may be used to point a parent type to a child type (eg, to enforce a specific derived implementation)
 		// We want to skip those redirects as the parent type still exists in this case
 		if (InField.IsA<UClass>() || InField.IsA<UScriptStruct>())
