Run a Python GC prior to the UE GC

This helps to clean-up any reference cycles that Python might have that are keeping UObject instances alive

diff --git a/Engine/Plugins/Experimental/PythonScriptPlugin/Source/PythonScriptPlugin/Private/PyUtil.cpp b/Engine/Plugins/Experimental/PythonScriptPlugin/Source/PythonScriptPlugin/Private/PyUtil.cpp
--- a/Engine/Plugins/Experimental/PythonScriptPlugin/Source/PythonScriptPlugin/Private/PyUtil.cpp
+++ b/Engine/Plugins/Experimental/PythonScriptPlugin/Source/PythonScriptPlugin/Private/PyUtil.cpp
@@ -1657,6 +1657,29 @@ bool ReThrowPythonError(FString* OutError)
 	return bFetchedError;
 }
 
+void CollectGarbage(int32 Generations)
+{
+	if (Generations == 2)
+	{
+		// The C API hardcodes Generations == 2 
+		PyGC_Collect();
+	}
+	else
+	{
+		// The Python API allows Generations to be overridden
+		if (FPyObjectPtr PyGCModule = FPyObjectPtr::StealReference(PyImport_ImportModule("gc")))
+		{
+			PyObject* PyGCDict = PyModule_GetDict(PyGCModule);
+
+			if (PyObject* PyCollectFunc = PyDict_GetItemString(PyGCDict, "collect"))
+			{
+				FPyObjectPtr PyCollectResult = FPyObjectPtr::StealReference(PyObject_CallFunction(PyCollectFunc, PyCStrCast("i"), Generations));
+				// PyCollectResult is unused, but the FPyObjectPtr handles its ref-count
+			}
+		}
+	}
+}
+
 }
 
 #endif	// WITH_PYTHON
diff --git a/Engine/Plugins/Experimental/PythonScriptPlugin/Source/PythonScriptPlugin/Private/PyUtil.h b/Engine/Plugins/Experimental/PythonScriptPlugin/Source/PythonScriptPlugin/Private/PyUtil.h
--- a/Engine/Plugins/Experimental/PythonScriptPlugin/Source/PythonScriptPlugin/Private/PyUtil.h
+++ b/Engine/Plugins/Experimental/PythonScriptPlugin/Source/PythonScriptPlugin/Private/PyUtil.h
@@ -652,6 +652,12 @@ namespace PyUtil
 	 * @return True if Python was in an error state, false otherwise.
 	 */
 	bool ReThrowPythonError(FString* OutError = nullptr);
+
+	/**
+	 * Run the Python garbage collector to clean-up objects that have internal reference cycles that can be broken.
+	 * @see [REDACTED_PROFILE_URL]/cpython/blob/main/InternalDocs/garbage_collector.md#optimization-generations
+	 */
+	void CollectGarbage(int32 Generations = 2);
 }
 
 #endif	// WITH_PYTHON
diff --git a/Engine/Plugins/Experimental/PythonScriptPlugin/Source/PythonScriptPlugin/Private/PythonScriptLibrary.cpp b/Engine/Plugins/Experimental/PythonScriptPlugin/Source/PythonScriptPlugin/Private/PythonScriptLibrary.cpp
--- a/Engine/Plugins/Experimental/PythonScriptPlugin/Source/PythonScriptPlugin/Private/PythonScriptLibrary.cpp
+++ b/Engine/Plugins/Experimental/PythonScriptPlugin/Source/PythonScriptPlugin/Private/PythonScriptLibrary.cpp
@@ -119,7 +119,7 @@ DEFINE_FUNCTION(UPythonScriptLibrary::execExecutePythonScript)
 		}
 
 		// Execute the Python command
-		FPyObjectPtr PyResult = FPyObjectPtr::StealReference(FPythonScriptPlugin::Get()->EvalString(*PythonScript, TEXT("<string>"), Py_file_input, PyTempGlobalDict, PyTempGlobalDict));
+		FPyObjectPtr PyResult = FPyObjectPtr::StealReference(FPythonScriptPlugin::Get()->EvalString(*PythonScript, TEXT("<string>"), Py_file_input, PyTempGlobalDict, PyTempLocalDict));
 
 		// Read the output values from the Python context
 		if (PyResult)
diff --git a/Engine/Plugins/Experimental/PythonScriptPlugin/Source/PythonScriptPlugin/Private/PythonScriptPlugin.cpp b/Engine/Plugins/Experimental/PythonScriptPlugin/Source/PythonScriptPlugin/Private/PythonScriptPlugin.cpp
--- a/Engine/Plugins/Experimental/PythonScriptPlugin/Source/PythonScriptPlugin/Private/PythonScriptPlugin.cpp
+++ b/Engine/Plugins/Experimental/PythonScriptPlugin/Source/PythonScriptPlugin/Private/PythonScriptPlugin.cpp
@@ -1139,6 +1139,9 @@ void FPythonScriptPlugin::InitializePython()
 	// We have to take this again prior to calling Py_Finalize, and all other code will lock on-demand via FPyScopedGIL
 	PyMainThreadState = PyEval_SaveThread();
 
+	// Hook into the UE GC so that we can clean-up Python cycles that may be keeping UObjects alive
+	FCoreUObjectDelegates::GetPreGarbageCollectDelegate().AddRaw(this, &FPythonScriptPlugin::OnPreGarbageCollect);
+
 #if WITH_EDITOR
 	// Initialize the Content Browser integration
 	if (GIsEditor && !IsRunningCommandlet() && GetDefault<UPythonScriptPluginUserSettings>()->bEnableContentBrowserIntegration)
@@ -1212,6 +1215,8 @@ void FPythonScriptPlugin::ShutdownPython()
 	// Notify any external listeners
 	OnPythonShutdownDelegate.Broadcast();
 
+	FCoreUObjectDelegates::GetPreGarbageCollectDelegate().RemoveAll(this);
+
 	FTSTicker::GetCoreTicker().RemoveTicker(TickHandle);
 	if (ModuleDelayedHandle.IsValid())
 	{
@@ -1893,6 +1898,12 @@ void FPythonScriptPlugin::OnAssetUpdated(const UObject* InObj)
 	}
 }
 
+void FPythonScriptPlugin::OnPreGarbageCollect()
+{
+	FPyScopedGIL GIL;
+	PyUtil::CollectGarbage();
+}
+
 #if WITH_EDITOR
 
 void FPythonScriptPlugin::OnPrepareToCleanseEditorObject(UObject* InObject)
diff --git a/Engine/Plugins/Experimental/PythonScriptPlugin/Source/PythonScriptPlugin/Private/PythonScriptPlugin.h b/Engine/Plugins/Experimental/PythonScriptPlugin/Source/PythonScriptPlugin/Private/PythonScriptPlugin.h
--- a/Engine/Plugins/Experimental/PythonScriptPlugin/Source/PythonScriptPlugin/Private/PythonScriptPlugin.h
+++ b/Engine/Plugins/Experimental/PythonScriptPlugin/Source/PythonScriptPlugin/Private/PythonScriptPlugin.h
@@ -194,6 +194,8 @@ private:
 
 	void OnAssetUpdated(const UObject* InObj);
 
+	void OnPreGarbageCollect();
+
 #if WITH_EDITOR
 	void OnPrepareToCleanseEditorObject(UObject* InObject);
 
