handle null property access while evaluating arguments for Execute Python Script

[RN] Fixed crash when an exception is encountered when evaluating the inputs of an Execute Python Script node

diff --git a/Engine/Plugins/Experimental/PythonScriptPlugin/Source/PythonScriptPlugin/Private/PythonScriptLibrary.cpp b/Engine/Plugins/Experimental/PythonScriptPlugin/Source/PythonScriptPlugin/Private/PythonScriptLibrary.cpp
--- a/Engine/Plugins/Experimental/PythonScriptPlugin/Source/PythonScriptPlugin/Private/PythonScriptLibrary.cpp
+++ b/Engine/Plugins/Experimental/PythonScriptPlugin/Source/PythonScriptPlugin/Private/PythonScriptLibrary.cpp
@@ -70,7 +70,18 @@ DEFINE_FUNCTION(UPythonScriptLibrary::execExecutePythonScript)
 				// Note: Medium term the Blueprint interpreter will change to provide us with a list of properties and 
 				// instance pointers, rather than forcing us to jump in and out of its interpreter loop (via StepCompiledIn)
 				Stack.StepCompiledIn<FProperty>(nullptr);
-				check(Stack.MostRecentProperty && Stack.MostRecentPropertyAddress);
+				if(!Stack.MostRecentProperty)
+				{
+#if WITH_PYTHON
+					// an exception was generated when evaluating the input, provide the user script with None:
+					PyDict_SetItemString(PyTempLocalDict, TCHAR_TO_UTF8(*PythonInput), Py_None);
+#endif	// WITH_PYTHON
+					continue;
+				}
+				else
+				{
+					check(Stack.MostRecentPropertyAddress);
+				}
 
 #if WITH_PYTHON
 				FPyObjectPtr PyInput;
