Don't call back into Python tick callbacks while user Python script is already running on the stack

This avoids issues of reentrancy that can confuse the Python exception handler, and allow things failing in the current script to also trigger exception states in the tick callback

diff --git a/Engine/Plugins/Experimental/PythonScriptPlugin/Source/PythonScriptPlugin/Private/PyGenUtil.cpp b/Engine/Plugins/Experimental/PythonScriptPlugin/Source/PythonScriptPlugin/Private/PyGenUtil.cpp
--- a/Engine/Plugins/Experimental/PythonScriptPlugin/Source/PythonScriptPlugin/Private/PyGenUtil.cpp
+++ b/Engine/Plugins/Experimental/PythonScriptPlugin/Source/PythonScriptPlugin/Private/PyGenUtil.cpp
@@ -1265,6 +1265,8 @@ bool InvokePythonCallableFromUnrealFunctionThunk(FPyObjectPtr InSelf, PyObject*
 		return false;
 	}
 
+	TGuardValue<bool> IsRunningUserScriptGuard(UE::Python::Private::GIsRunningUserScript, true);
+
 	// Prepare the arguments tuple for the Python callable
 	FPyObjectPtr PyArgs;
 	if (InSelf || PyParams.Num() > 0)
diff --git a/Engine/Plugins/Experimental/PythonScriptPlugin/Source/PythonScriptPlugin/Private/PySlate.cpp b/Engine/Plugins/Experimental/PythonScriptPlugin/Source/PythonScriptPlugin/Private/PySlate.cpp
--- a/Engine/Plugins/Experimental/PythonScriptPlugin/Source/PythonScriptPlugin/Private/PySlate.cpp
+++ b/Engine/Plugins/Experimental/PythonScriptPlugin/Source/PythonScriptPlugin/Private/PySlate.cpp
@@ -29,6 +29,13 @@ FPyDelegateHandle* RegisterSlateTickCallback(FSlateApplication::FSlateTickEvent&
 	// We copy the PyCallable into the lambda to keep the Python object alive as long as the delegate is bound
 	FDelegateHandle TickEventDelegateHandle = InSlateTickEvent.AddLambda([PyCallable = FPyAutoGILPtr(FPyObjectPtr::NewReference(InPyCallable))](const float InDeltaTime) mutable
 	{
+		// Do not tick into Python while Python is already running on the stack
+		// Doing so could cause re-entrancy issues and confuse exception handling
+		if (UE::Python::Private::GIsRunningUserScript)
+		{
+			return;
+		}
+
 		// Do not tick into Python when it may not be safe to call back into C++
 		if (GIsSavingPackage || IsGarbageCollecting() || FUObjectThreadContext::Get().IsRoutingPostLoad)
 		{
diff --git a/Engine/Plugins/Experimental/PythonScriptPlugin/Source/PythonScriptPlugin/Private/PyUtil.cpp b/Engine/Plugins/Experimental/PythonScriptPlugin/Source/PythonScriptPlugin/Private/PyUtil.cpp
--- a/Engine/Plugins/Experimental/PythonScriptPlugin/Source/PythonScriptPlugin/Private/PyUtil.cpp
+++ b/Engine/Plugins/Experimental/PythonScriptPlugin/Source/PythonScriptPlugin/Private/PyUtil.cpp
@@ -46,6 +46,13 @@ DEFINE_LOG_CATEGORY(LogPython);
 
 #if WITH_PYTHON
 
+namespace UE::Python::Private
+{
+
+bool GIsRunningUserScript = false;
+
+}
+
 namespace PyUtil
 {
 
diff --git a/Engine/Plugins/Experimental/PythonScriptPlugin/Source/PythonScriptPlugin/Private/PyUtil.h b/Engine/Plugins/Experimental/PythonScriptPlugin/Source/PythonScriptPlugin/Private/PyUtil.h
--- a/Engine/Plugins/Experimental/PythonScriptPlugin/Source/PythonScriptPlugin/Private/PyUtil.h
+++ b/Engine/Plugins/Experimental/PythonScriptPlugin/Source/PythonScriptPlugin/Private/PyUtil.h
@@ -25,6 +25,19 @@ struct FPropertyAccessChangeNotify;
 /** Cast a string literal to a char* (via a void* to avoid a compiler warning) */
 #define PyCStrCast(STR) (char*)(void*)(STR)
 
+namespace UE::Python::Private
+{
+
+/**
+ * Set to true whenever we're running user-defined Python script.
+ * This is primarily set for FPythonScriptPlugin::EvalString, but should 
+ * also be set whenever calling back into user-defined Python from other
+ * places (eg, PythonGeneratedClass functions, delegates, etc).
+ */
+extern bool GIsRunningUserScript;
+
+}
+
 namespace PyUtil
 {
 	/** Types used by various APIs that change between versions */
diff --git a/Engine/Plugins/Experimental/PythonScriptPlugin/Source/PythonScriptPlugin/Private/PyWrapperObject.cpp b/Engine/Plugins/Experimental/PythonScriptPlugin/Source/PythonScriptPlugin/Private/PyWrapperObject.cpp
--- a/Engine/Plugins/Experimental/PythonScriptPlugin/Source/PythonScriptPlugin/Private/PyWrapperObject.cpp
+++ b/Engine/Plugins/Experimental/PythonScriptPlugin/Source/PythonScriptPlugin/Private/PyWrapperObject.cpp
@@ -1991,6 +1991,8 @@ void UPythonGeneratedClass::PostInitInstance(UObject* InObj, FObjectInstancingGr
 
 		if (PyPostInitFunction)
 		{
+			TGuardValue<bool> IsRunningUserScriptGuard(UE::Python::Private::GIsRunningUserScript, true);
+
 			FPyObjectPtr PySelf = FPyObjectPtr::StealReference((PyObject*)FPyWrapperObjectFactory::Get().CreateInstance(InObj));
 			if (PySelf && ensureAlwaysMsgf(PySelf->ob_type == PyType, TEXT("Object '%s' (class: %s) had an unexpected PyType when calling PostInitInstance! Self is '%s' but we expected '%s'"), *InObj->GetPathName(), *GetPathName(), *PyUtil::GetFriendlyTypename(PySelf), *PyUtil::GetFriendlyTypename(PyType)))
 			{
diff --git a/Engine/Plugins/Experimental/PythonScriptPlugin/Source/PythonScriptPlugin/Private/PyWrapperStruct.cpp b/Engine/Plugins/Experimental/PythonScriptPlugin/Source/PythonScriptPlugin/Private/PyWrapperStruct.cpp
--- a/Engine/Plugins/Experimental/PythonScriptPlugin/Source/PythonScriptPlugin/Private/PyWrapperStruct.cpp
+++ b/Engine/Plugins/Experimental/PythonScriptPlugin/Source/PythonScriptPlugin/Private/PyWrapperStruct.cpp
@@ -1651,6 +1651,8 @@ void UPythonGeneratedStruct::InitializeStruct(void* Dest, int32 ArrayDim) const
 
 		if (PyPostInitFunction)
 		{
+			TGuardValue<bool> IsRunningUserScriptGuard(UE::Python::Private::GIsRunningUserScript, true);
+
 			const int32 Stride = GetStructureSize();
 			for (int32 ArrIndex = 0; ArrIndex < ArrayDim; ++ArrIndex)
 			{
diff --git a/Engine/Plugins/Experimental/PythonScriptPlugin/Source/PythonScriptPlugin/Private/PythonScriptPlugin.cpp b/Engine/Plugins/Experimental/PythonScriptPlugin/Source/PythonScriptPlugin/Private/PythonScriptPlugin.cpp
--- a/Engine/Plugins/Experimental/PythonScriptPlugin/Source/PythonScriptPlugin/Private/PythonScriptPlugin.cpp
+++ b/Engine/Plugins/Experimental/PythonScriptPlugin/Source/PythonScriptPlugin/Private/PythonScriptPlugin.cpp
@@ -1562,6 +1562,7 @@ PyObject* FPythonScriptPlugin::EvalString(const TCHAR* InStr, const TCHAR* InCon
 	}
 
 	FScopedEncodingGuard EncodingGuard;
+	TGuardValue<bool> IsRunningUserScriptGuard(UE::Python::Private::GIsRunningUserScript, true);
 	PyUtil::FEvalStack::Get().PushContext(PyUtil::FEvalStack::FEvalContext{ InContext, InGlobalDict, InLocalDict});
 	PyObject* PyEvalResult = PyEval_EvalCode((PyUtil::FPyCodeObjectType*)PyCodeObj.Get(), InGlobalDict, InLocalDict);
 	PyUtil::FEvalStack::Get().PopContext();
