Add check for failure to get PyWrapperMetadata in FPyWrapperBaseMetaData::SetMetaData

diff --git a/Engine/Plugins/Experimental/PythonScriptPlugin/Source/PythonScriptPlugin/Private/PyWrapperBase.cpp b/Engine/Plugins/Experimental/PythonScriptPlugin/Source/PythonScriptPlugin/Private/PyWrapperBase.cpp
--- a/Engine/Plugins/Experimental/PythonScriptPlugin/Source/PythonScriptPlugin/Private/PyWrapperBase.cpp
+++ b/Engine/Plugins/Experimental/PythonScriptPlugin/Source/PythonScriptPlugin/Private/PyWrapperBase.cpp
@@ -3,6 +3,7 @@
 #include "PyWrapperBase.h"
 #include "PyReferenceCollector.h"
 #include "PyGIL.h"
+#include "PyUtil.h"
 
 #include UE_INLINE_GENERATED_CPP_BY_NAME(PyWrapperBase)
 
@@ -123,6 +124,9 @@ void FPyWrapperBaseMetaData::SetMetaData(PyTypeObject* PyType, FPyWrapperBaseMet
 		if (!PyWrapperMetaData)
 		{
 			PyWrapperMetaData = (FPyWrapperBaseMetaDataObject*)PyObject_CallObject((PyObject*)&PyWrapperBaseMetaDataType, nullptr);
+			checkf(PyWrapperMetaData, TEXT("PyWrapperMetaData is null for PyType: %s"),
+				*PyUtil::GetFriendlyTypename(PyType));
+
 			PyDict_SetItemString(PyType->tp_dict, "_wrapper_meta_data", (PyObject*)PyWrapperMetaData);
 			Py_DECREF(PyWrapperMetaData);
 		}
