Fix crash when calling ExecPythonCommandEX before Python is initialized

diff --git a/Engine/Plugins/Experimental/PythonScriptPlugin/Source/PythonScriptPlugin/Private/PythonScriptPlugin.cpp b/Engine/Plugins/Experimental/PythonScriptPlugin/Source/PythonScriptPlugin/Private/PythonScriptPlugin.cpp
--- a/Engine/Plugins/Experimental/PythonScriptPlugin/Source/PythonScriptPlugin/Private/PythonScriptPlugin.cpp
+++ b/Engine/Plugins/Experimental/PythonScriptPlugin/Source/PythonScriptPlugin/Private/PythonScriptPlugin.cpp
@@ -687,6 +687,15 @@ bool FPythonScriptPlugin::ExecPythonCommandEx(FPythonCommandEx& InOutPythonComma
 	}
 
 #if WITH_PYTHON
+	if (!bInitialized)
+	{
+		InOutPythonCommand.CommandResult =
+			TEXT("Attempt to execute python command before PythonScriptPlugin is initialized. Ensure your call is after OnPythonInitialized.");
+
+		UE_LOG(LogPython, Warning, TEXT("%s"), *InOutPythonCommand.CommandResult);
+		return false;
+	}
+
 	if (InOutPythonCommand.ExecutionMode == EPythonCommandExecutionMode::ExecuteFile)
 	{
 		// The EPythonCommandExecutionMode::ExecuteFile name is misleading as it is used to run literal code or a .py file. Detect
