Remove IsPythonAvailable check from -ExecutePythonScript until IPythonScriptPlugin has a better API for handling deferred config/initialization

#rnx

diff --git a/Engine/Plugins/Experimental/PythonScriptPlugin/Source/PythonScriptPlugin/Private/EditorUtilities/EditorPythonExecuter.cpp b/Engine/Plugins/Experimental/PythonScriptPlugin/Source/PythonScriptPlugin/Private/EditorUtilities/EditorPythonExecuter.cpp
--- a/Engine/Plugins/Experimental/PythonScriptPlugin/Source/PythonScriptPlugin/Private/EditorUtilities/EditorPythonExecuter.cpp
+++ b/Engine/Plugins/Experimental/PythonScriptPlugin/Source/PythonScriptPlugin/Private/EditorUtilities/EditorPythonExecuter.cpp
@@ -187,10 +187,6 @@ void FEditorPythonExecuter::OnStartupModule()
 		{
 			UE_LOG(LogEditorPythonExecuter, Error, TEXT("-ExecutePythonScript cannot be used by a commandlet."));
 		}
-		else if (!IPythonScriptPlugin::Get()->IsPythonAvailable())
-		{
-			UE_LOG(LogEditorPythonExecuter, Error, TEXT("-ExecutePythonScript cannot be used when Python support is disabled."));
-		}
 		else
 		{
 			IPythonScriptPlugin::Get()->OnPythonInitialized().AddLambda([ScriptAndArgs = MoveTemp(ScriptAndArgs), bScriptErrorsAreFatal = FParse::Param(FCommandLine::Get(), TEXT("ScriptErrorsAreFatal"))]() mutable
