Improved Python initialization API for deferred initialization

* Added IPythonScriptPlugin::IsPythonConfigured and IPythonScriptPlugin::OnPythonConfigured to know when Python has been configured by deferred initialization.
* Added IPythonScriptPlugin::IsPythonInitialized to compliment IPythonScriptPlugin::OnPythonInitialized for deferred initialization.
* Added IPythonScriptPlugin::RegisterOnPythonConfigured and IPythonScriptPlugin::RegisterOnPythonInitialized as convenience wrappers around their respective IsX and OnX functions.
* Fixed -ExecutePythonScript failing with deferred initialization by updating it to use these new APIs.

diff --git a/Engine/Plugins/Experimental/PythonScriptPlugin/Source/PythonScriptPlugin/Private/EditorUtilities/EditorPythonExecuter.cpp b/Engine/Plugins/Experimental/PythonScriptPlugin/Source/PythonScriptPlugin/Private/EditorUtilities/EditorPythonExecuter.cpp
--- a/Engine/Plugins/Experimental/PythonScriptPlugin/Source/PythonScriptPlugin/Private/EditorUtilities/EditorPythonExecuter.cpp
+++ b/Engine/Plugins/Experimental/PythonScriptPlugin/Source/PythonScriptPlugin/Private/EditorUtilities/EditorPythonExecuter.cpp
@@ -189,11 +189,21 @@ void FEditorPythonExecuter::OnStartupModule()
 		}
 		else
 		{
-			IPythonScriptPlugin::Get()->OnPythonInitialized().AddLambda([ScriptAndArgs = MoveTemp(ScriptAndArgs), bScriptErrorsAreFatal = FParse::Param(FCommandLine::Get(), TEXT("ScriptErrorsAreFatal"))]() mutable
+			IPythonScriptPlugin::Get()->RegisterOnPythonConfigured(FSimpleDelegate::CreateLambda([ScriptAndArgs = MoveTemp(ScriptAndArgs), bScriptErrorsAreFatal = FParse::Param(FCommandLine::Get(), TEXT("ScriptErrorsAreFatal"))]() mutable
 			{
-				InternalEditorPythonRunner::CreateNotification(ScriptAndArgs);
-				InternalEditorPythonRunner::Executer = MakeUnique<InternalEditorPythonRunner::FExecuterTickable>(MoveTemp(ScriptAndArgs), bScriptErrorsAreFatal);
-			});
+				if (IPythonScriptPlugin::Get()->IsPythonAvailable())
+				{
+					IPythonScriptPlugin::Get()->RegisterOnPythonInitialized(FSimpleDelegate::CreateLambda([ScriptAndArgs = MoveTemp(ScriptAndArgs), bScriptErrorsAreFatal]() mutable
+					{
+						InternalEditorPythonRunner::CreateNotification(ScriptAndArgs);
+						InternalEditorPythonRunner::Executer = MakeUnique<InternalEditorPythonRunner::FExecuterTickable>(MoveTemp(ScriptAndArgs), bScriptErrorsAreFatal);
+					}));
+				}
+				else
+				{
+					UE_LOG(LogEditorPythonExecuter, Error, TEXT("-ExecutePythonScript cannot be used when Python support is disabled."));
+				}
+			}));
 		}
 	}
 }
diff --git a/Engine/Plugins/Experimental/PythonScriptPlugin/Source/PythonScriptPlugin/Private/PythonScriptPlugin.cpp b/Engine/Plugins/Experimental/PythonScriptPlugin/Source/PythonScriptPlugin/Private/PythonScriptPlugin.cpp
--- a/Engine/Plugins/Experimental/PythonScriptPlugin/Source/PythonScriptPlugin/Private/PythonScriptPlugin.cpp
+++ b/Engine/Plugins/Experimental/PythonScriptPlugin/Source/PythonScriptPlugin/Private/PythonScriptPlugin.cpp
@@ -668,6 +668,24 @@ bool FPythonScriptPlugin::IsPythonAvailable() const
 #endif
 }
 
+bool FPythonScriptPlugin::IsPythonConfigured() const
+{
+#if WITH_PYTHON
+	return bIsConfigured;
+#else
+	return false;
+#endif
+}
+
+bool FPythonScriptPlugin::IsPythonInitialized() const
+{
+#if WITH_PYTHON
+	return bIsFullyInitialized;
+#else
+	return false;
+#endif
+}
+
 bool FPythonScriptPlugin::ExecPythonCommand(const TCHAR* InPythonCommand)
 {
 	FPythonCommandEx PythonCommand;
@@ -687,7 +705,7 @@ bool FPythonScriptPlugin::ExecPythonCommandEx(FPythonCommandEx& InOutPythonComma
 	}
 
 #if WITH_PYTHON
-	if (!bInitialized)
+	if (!bIsInterpreterInitialized)
 	{
 		InOutPythonCommand.CommandResult =
 			TEXT("Attempt to execute python command before PythonScriptPlugin is initialized. Ensure your call is after OnPythonInitialized.");
@@ -800,6 +818,11 @@ FString FPythonScriptPlugin::GetInterpreterExecutablePath() const
 #endif	// WITH_PYTHON
 }
 
+FSimpleMulticastDelegate& FPythonScriptPlugin::OnPythonConfigured()
+{
+	return OnPythonConfiguredDelegate;
+}
+
 FSimpleMulticastDelegate& FPythonScriptPlugin::OnPythonInitialized()
 {
 	return OnPythonInitializedDelegate;
@@ -832,6 +855,10 @@ void FPythonScriptPlugin::OnPostEngineInit()
 	// IsPythonAvailable() will always return false prior to configuring bIsEnabled
 	bIsEnabled = UE::Python::Private::IsPythonEnabled();
 	
+	bIsConfigured = true;
+	OnPythonConfiguredDelegate.Broadcast();
+	OnPythonConfiguredDelegate.Clear();
+
 	if (!IsPythonAvailable())
 	{
 		return;
@@ -925,7 +952,7 @@ void FPythonScriptPlugin::InitializePython()
 {
 	TRACE_CPUPROFILER_EVENT_SCOPE(FPythonScriptPlugin::InitializePython)
 
-	bInitialized = true;
+	bIsInterpreterInitialized = true;
 
 	FScopedSlowTask SlowTask(1, LOCTEXT("InitializingPython", "Initializing Python..."));
 	SlowTask.Visibility = ESlowTaskVisibility::Important; // this function can be very slow, users will benefit from our messages
@@ -1206,7 +1233,7 @@ void FPythonScriptPlugin::InitializePython()
 
 void FPythonScriptPlugin::ShutdownPython()
 {
-	if (!bInitialized)
+	if (!bIsInterpreterInitialized)
 	{
 		return;
 	}
@@ -1223,6 +1250,7 @@ void FPythonScriptPlugin::ShutdownPython()
 
 	// Notify any external listeners
 	OnPythonShutdownDelegate.Broadcast();
+	OnPythonShutdownDelegate.Clear();
 
 	FCoreUObjectDelegates::GetPreGarbageCollectDelegate().RemoveAll(this);
 
@@ -1272,9 +1300,10 @@ void FPythonScriptPlugin::ShutdownPython()
 
 	Py_Finalize();
 
-	bInitialized = false;
+	bIsConfigured = false;
+	bIsInterpreterInitialized = false;
+	bIsFullyInitialized = false;
 	bIsEnabled = false;
-	bRanStartupScripts = false;
 }
 
 void FPythonScriptPlugin::InitPipInstaller()
@@ -1366,7 +1395,7 @@ void FPythonScriptPlugin::RunPipInstaller()
 void FPythonScriptPlugin::RequestStubCodeGeneration()
 {
 	// Ignore requests made before the startup scripts have ran
-	if (!bRanStartupScripts)
+	if (!bIsFullyInitialized)
 	{
 		return;
 	}
@@ -1412,12 +1441,12 @@ void FPythonScriptPlugin::GenerateStubCode()
 
 void FPythonScriptPlugin::RunStartupScripts()
 {
-	if (bRanStartupScripts)
+	if (bIsFullyInitialized)
 	{
 		return;
 	}
 
-	bRanStartupScripts = true;
+	bIsFullyInitialized = true;
 
 	// Run start-up scripts now
 	TArray<FString> PySysPaths;
@@ -1462,6 +1491,7 @@ void FPythonScriptPlugin::RunStartupScripts()
 
 	// Notify any external listeners
 	OnPythonInitializedDelegate.Broadcast();
+	OnPythonInitializedDelegate.Clear();
 
 #if WITH_EDITOR
 	// Activate the Content Browser integration (now that editor subsystems are available)
diff --git a/Engine/Plugins/Experimental/PythonScriptPlugin/Source/PythonScriptPlugin/Private/PythonScriptPlugin.h b/Engine/Plugins/Experimental/PythonScriptPlugin/Source/PythonScriptPlugin/Private/PythonScriptPlugin.h
--- a/Engine/Plugins/Experimental/PythonScriptPlugin/Source/PythonScriptPlugin/Private/PythonScriptPlugin.h
+++ b/Engine/Plugins/Experimental/PythonScriptPlugin/Source/PythonScriptPlugin/Private/PythonScriptPlugin.h
@@ -90,7 +90,7 @@ struct IPythonCommandMenu
 };
 #endif	// WITH_PYTHON
 
-class FPythonScriptPlugin 
+class FPythonScriptPlugin final
 	: public IPythonScriptPlugin
 	, public FSelfRegisteringExec
 	, public FEnumEditorUtils::INotifyOnEnumChanged
@@ -106,9 +106,12 @@ public:
 
 	//~ IPythonScriptPlugin interface
 	virtual bool IsPythonAvailable() const override;
+	virtual bool IsPythonConfigured() const override;
+	virtual bool IsPythonInitialized() const override;
 	virtual bool ExecPythonCommand(const TCHAR* InPythonCommand) override;
 	virtual bool ExecPythonCommandEx(FPythonCommandEx& InOutPythonCommand) override;
 	virtual FString GetInterpreterExecutablePath() const override;
+	virtual FSimpleMulticastDelegate& OnPythonConfigured() override;
 	virtual FSimpleMulticastDelegate& OnPythonInitialized() override;
 	virtual FSimpleMulticastDelegate& OnPythonShutdown() override;
 
@@ -219,9 +222,10 @@ private:
 	FPyObjectPtr PyConsoleLocalDict;
 	FPyObjectPtr PyUnrealModule;
 	PyThreadState* PyMainThreadState = nullptr;
-	bool bInitialized = false;
+	bool bIsConfigured = false;
+	bool bIsInterpreterInitialized = false; // Python interpreter is available; may not yet have run startup scripts
+	bool bIsFullyInitialized = false; // Python interpreter is available, and start-up scripts have run
 	bool bIsEnabled = false;
-	bool bRanStartupScripts = false;
 
 #if PY_MAJOR_VERSION >= 3 && PY_MINOR_VERSION >= 11
 	PyConfig ModulePyConfig;
@@ -229,6 +233,7 @@ private:
 #endif // 3.11 upgrade
 #endif	// WITH_PYTHON
 
+	FSimpleMulticastDelegate OnPythonConfiguredDelegate;
 	FSimpleMulticastDelegate OnPythonInitializedDelegate;
 	FSimpleMulticastDelegate OnPythonShutdownDelegate;
 };
diff --git a/Engine/Plugins/Experimental/PythonScriptPlugin/Source/PythonScriptPlugin/Public/IPythonScriptPlugin.h b/Engine/Plugins/Experimental/PythonScriptPlugin/Source/PythonScriptPlugin/Public/IPythonScriptPlugin.h
--- a/Engine/Plugins/Experimental/PythonScriptPlugin/Source/PythonScriptPlugin/Public/IPythonScriptPlugin.h
+++ b/Engine/Plugins/Experimental/PythonScriptPlugin/Source/PythonScriptPlugin/Public/IPythonScriptPlugin.h
@@ -18,10 +18,22 @@ public:
 	}
 
 	/**
-	 * Check to see whether the plugin was built with Python support enabled.
+	 * Check to see whether the plugin has Python support enabled.
+	 * @note This may return false until IsPythonConfigured is true.
 	 */
 	virtual bool IsPythonAvailable() const = 0;
 
+	/**
+	 * Check to see whether Python has been configured.
+	 * @note Python may be configured but not yet be initialized (@see IsPythonInitialized).
+	 */
+	virtual bool IsPythonConfigured() const = 0;
+
+	/**
+	 * Check to see whether Python has been initialized and is ready to use.
+	 */
+	virtual bool IsPythonInitialized() const = 0;
+
 	/**
 	 * Execute the given Python command.
 	 * This may be literal Python code, or a file (with optional arguments) that you want to run.
@@ -40,6 +52,11 @@ public:
 	 */
 	virtual FString GetInterpreterExecutablePath() const = 0;
 
+	/**
+	 * Delegate called after Python has been configured.
+	 */
+	virtual FSimpleMulticastDelegate& OnPythonConfigured() = 0;
+
 	/**
 	 * Delegate called after Python has been initialized.
 	 */
@@ -49,4 +66,34 @@ public:
 	 * Delegate called before Python is shutdown.
 	 */
 	virtual FSimpleMulticastDelegate& OnPythonShutdown() = 0;
+
+	/**
+	 * Wrapper around OnPythonConfigured that will either register the callback, or call it immediately if IsPythonConfigured is already true.
+	 */
+	void RegisterOnPythonConfigured(FSimpleDelegate&& Callback)
+	{
+		if (IsPythonConfigured())
+		{
+			Callback.ExecuteIfBound();
+		}
+		else
+		{
+			OnPythonConfigured().Add(MoveTemp(Callback));
+		}
+	}
+
+	/**
+	 * Wrapper around OnPythonInitialized that will either register the callback, or call it immediately if IsPythonInitialized is already true.
+	 */
+	void RegisterOnPythonInitialized(FSimpleDelegate&& Callback)
+	{
+		if (IsPythonInitialized())
+		{
+			Callback.ExecuteIfBound();
+		}
+		else
+		{
+			OnPythonInitialized().Add(MoveTemp(Callback));
+		}
+	}
 };
