Keep legacy Python type redirectors alive through a GC

diff --git a/Engine/Plugins/Experimental/PythonScriptPlugin/Source/PythonScriptPlugin/Private/PyWrapperEnum.cpp b/Engine/Plugins/Experimental/PythonScriptPlugin/Source/PythonScriptPlugin/Private/PyWrapperEnum.cpp
--- a/Engine/Plugins/Experimental/PythonScriptPlugin/Source/PythonScriptPlugin/Private/PyWrapperEnum.cpp
+++ b/Engine/Plugins/Experimental/PythonScriptPlugin/Source/PythonScriptPlugin/Private/PyWrapperEnum.cpp
@@ -762,9 +762,10 @@ public:
 		}
 		NewEnum->EnumValueDefs.Reset();
 
-		if (EnumRedirector)
+		NewEnum->EnumRedirector = EnumRedirector;
+		if (NewEnum->EnumRedirector)
 		{
-			EnumRedirector->DestinationObject = NewEnum;
+			NewEnum->EnumRedirector->DestinationObject = NewEnum;
 		}
 	}
 
@@ -884,6 +885,14 @@ private:
 	UPythonGeneratedEnum* NewEnum;
 };
 
+void UPythonGeneratedEnum::AddReferencedObjects(UObject* InThis, FReferenceCollector& Collector)
+{
+	Super::AddReferencedObjects(InThis, Collector);
+
+	UPythonGeneratedEnum* This = CastChecked<UPythonGeneratedEnum>(InThis);
+	Collector.AddReferencedObject(This->EnumRedirector, This);
+}
+
 void UPythonGeneratedEnum::BeginDestroy()
 {
 	ReleasePythonResources();
diff --git a/Engine/Plugins/Experimental/PythonScriptPlugin/Source/PythonScriptPlugin/Private/PyWrapperEnum.h b/Engine/Plugins/Experimental/PythonScriptPlugin/Source/PythonScriptPlugin/Private/PyWrapperEnum.h
--- a/Engine/Plugins/Experimental/PythonScriptPlugin/Source/PythonScriptPlugin/Private/PyWrapperEnum.h
+++ b/Engine/Plugins/Experimental/PythonScriptPlugin/Source/PythonScriptPlugin/Private/PyWrapperEnum.h
@@ -6,6 +6,8 @@
 #include "UObject/Class.h"
 #include "PyWrapperEnum.generated.h"
 
+class UObjectRedirector;
+
 #if WITH_PYTHON
 
 /** Python type for FPyWrapperEnum */
@@ -121,6 +123,7 @@ class UPythonGeneratedEnum final : public UEnum, public IPythonResourceOwner
 
 public:
 	//~ UObject interface
+	static void AddReferencedObjects(UObject* InThis, FReferenceCollector& Collector);
 	virtual void BeginDestroy() override;
 	virtual bool IsAsset() const override
 	{
@@ -162,6 +165,8 @@ private:
 	/** Meta-data for this generated enum that is applied to the Python type */
 	FPyWrapperEnumMetaData PyMetaData;
 
+	TObjectPtr<UObjectRedirector> EnumRedirector;
+
 	friend class FPythonGeneratedEnumBuilder;
 
 #else	// WITH_PYTHON
diff --git a/Engine/Plugins/Experimental/PythonScriptPlugin/Source/PythonScriptPlugin/Private/PyWrapperObject.cpp b/Engine/Plugins/Experimental/PythonScriptPlugin/Source/PythonScriptPlugin/Private/PyWrapperObject.cpp
--- a/Engine/Plugins/Experimental/PythonScriptPlugin/Source/PythonScriptPlugin/Private/PyWrapperObject.cpp
+++ b/Engine/Plugins/Experimental/PythonScriptPlugin/Source/PythonScriptPlugin/Private/PyWrapperObject.cpp
@@ -1421,9 +1421,10 @@ public:
 		}
 
 		Py_BEGIN_ALLOW_THREADS
-		if (UObjectRedirector* ClassRedirector = CreatePythonTypeLegacyRedirector(PyUtil::GetCleanTypename(PyType), FTopLevelAssetPath(NewClass->GetOuter()->GetFName(), *ClassName)))
+		NewClass->ClassRedirector = CreatePythonTypeLegacyRedirector(PyUtil::GetCleanTypename(PyType), FTopLevelAssetPath(NewClass->GetOuter()->GetFName(), *ClassName));
+		if (NewClass->ClassRedirector)
 		{
-			ClassRedirector->DestinationObject = NewClass;
+			NewClass->ClassRedirector->DestinationObject = NewClass;
 		}
 		NewClass->Rename(*ClassName, nullptr, REN_DontCreateRedirectors);
 
@@ -1970,6 +1971,14 @@ private:
 	UPythonGeneratedClass* NewClass;
 };
 
+void UPythonGeneratedClass::AddReferencedObjects(UObject* InThis, FReferenceCollector& Collector)
+{
+	Super::AddReferencedObjects(InThis, Collector);
+
+	UPythonGeneratedClass* This = CastChecked<UPythonGeneratedClass>(InThis);
+	Collector.AddReferencedObject(This->ClassRedirector, This);
+}
+
 void UPythonGeneratedClass::PostRename(UObject* OldOuter, const FName OldName)
 {
 	Super::PostRename(OldOuter, OldName);
diff --git a/Engine/Plugins/Experimental/PythonScriptPlugin/Source/PythonScriptPlugin/Private/PyWrapperObject.h b/Engine/Plugins/Experimental/PythonScriptPlugin/Source/PythonScriptPlugin/Private/PyWrapperObject.h
--- a/Engine/Plugins/Experimental/PythonScriptPlugin/Source/PythonScriptPlugin/Private/PyWrapperObject.h
+++ b/Engine/Plugins/Experimental/PythonScriptPlugin/Source/PythonScriptPlugin/Private/PyWrapperObject.h
@@ -8,6 +8,8 @@
 #include "UObject/PropertyAccessUtil.h"
 #include "PyWrapperObject.generated.h"
 
+class UObjectRedirector;
+
 #if WITH_PYTHON
 
 /** Python type for FPyWrapperObject */
@@ -192,6 +194,7 @@ class UPythonGeneratedClass final : public UClass, public IPythonResourceOwner
 
 public:
 	//~ UObject interface
+	static void AddReferencedObjects(UObject* InThis, FReferenceCollector& Collector);
 	virtual void PostRename(UObject* OldOuter, const FName OldName) override;
 	virtual void BeginDestroy() override;
 	virtual bool IsAsset() const override
@@ -238,6 +241,8 @@ private:
 	/** Meta-data for this generated class that is applied to the Python type */
 	FPyWrapperObjectMetaData PyMetaData;
 
+	TObjectPtr<UObjectRedirector> ClassRedirector;
+
 	friend class FPythonGeneratedClassBuilder;
 
 #else	// WITH_PYTHON
diff --git a/Engine/Plugins/Experimental/PythonScriptPlugin/Source/PythonScriptPlugin/Private/PyWrapperStruct.cpp b/Engine/Plugins/Experimental/PythonScriptPlugin/Source/PythonScriptPlugin/Private/PyWrapperStruct.cpp
--- a/Engine/Plugins/Experimental/PythonScriptPlugin/Source/PythonScriptPlugin/Private/PyWrapperStruct.cpp
+++ b/Engine/Plugins/Experimental/PythonScriptPlugin/Source/PythonScriptPlugin/Private/PyWrapperStruct.cpp
@@ -1519,9 +1519,10 @@ public:
 			PrepareOldStructForReinstancing();
 		}
 
-		if (UObjectRedirector* StructRedirector = CreatePythonTypeLegacyRedirector(PyUtil::GetCleanTypename(PyType), FTopLevelAssetPath(NewStruct->GetOuter()->GetFName(), *StructName)))
+		NewStruct->StructRedirector = CreatePythonTypeLegacyRedirector(PyUtil::GetCleanTypename(PyType), FTopLevelAssetPath(NewStruct->GetOuter()->GetFName(), *StructName));
+		if (NewStruct->StructRedirector)
 		{
-			StructRedirector->DestinationObject = NewStruct;
+			NewStruct->StructRedirector->DestinationObject = NewStruct;
 		}
 		NewStruct->Rename(*StructName, nullptr, REN_DontCreateRedirectors);
 
@@ -1630,6 +1631,14 @@ private:
 	UPythonGeneratedStruct* NewStruct;
 };
 
+void UPythonGeneratedStruct::AddReferencedObjects(UObject* InThis, FReferenceCollector& Collector)
+{
+	Super::AddReferencedObjects(InThis, Collector);
+
+	UPythonGeneratedStruct* This = CastChecked<UPythonGeneratedStruct>(InThis);
+	Collector.AddReferencedObject(This->StructRedirector, This);
+}
+
 void UPythonGeneratedStruct::PostRename(UObject* OldOuter, const FName OldName)
 {
 	Super::PostRename(OldOuter, OldName);
diff --git a/Engine/Plugins/Experimental/PythonScriptPlugin/Source/PythonScriptPlugin/Private/PyWrapperStruct.h b/Engine/Plugins/Experimental/PythonScriptPlugin/Source/PythonScriptPlugin/Private/PyWrapperStruct.h
--- a/Engine/Plugins/Experimental/PythonScriptPlugin/Source/PythonScriptPlugin/Private/PyWrapperStruct.h
+++ b/Engine/Plugins/Experimental/PythonScriptPlugin/Source/PythonScriptPlugin/Private/PyWrapperStruct.h
@@ -11,6 +11,8 @@
 #include "Templates/TypeCompatibleBytes.h"
 #include "PyWrapperStruct.generated.h"
 
+class UObjectRedirector;
+
 #if WITH_PYTHON
 
 /** Python type for FPyWrapperStruct */
@@ -380,6 +382,7 @@ class UPythonGeneratedStruct final : public UScriptStruct, public IPythonResourc
 
 public:
 	//~ UObject interface
+	static void AddReferencedObjects(UObject* InThis, FReferenceCollector& Collector);
 	virtual void PostRename(UObject* OldOuter, const FName OldName) override;
 	virtual void BeginDestroy() override;
 	virtual bool IsAsset() const override
@@ -412,6 +415,8 @@ private:
 	/** Meta-data for this generated struct that is applied to the Python type */
 	FPyWrapperStructMetaData PyMetaData;
 
+	TObjectPtr<UObjectRedirector> StructRedirector;
+
 	friend class FPythonGeneratedStructBuilder;
 
 #else	// WITH_PYTHON
