Added support for lazy Python initialization, when running with 'Engine.Python.IsEnabledByDefault' set to false

This allows tools that require Python to request that it be enabled and initialized on-demand at runtime, rather than require everyone have Python enabled in-case they want to use any tooling that uses Python. This can be triggered by calling IPythonScriptPlugin::ForceEnablePythonAtRuntime (from C++) or UPythonScriptLibrary::ForceEnablePythonAtRuntime (from BP).

ForceEnablePythonAtRuntime will flag Python as being enabled, and can be called at any point. If called prior to Python being ready to initialize, it will just set a flag for it to use to enable itself when it's ready. If Python has already passed the point where it would have naturally initialized (past OnPostEngineInit) then it will immediately initialize Python. Similarly, if Python has passed the point where it would have naturally run its start-up scripts (the first Tick) then it will immediately run the start-up scripts.

Running the Python commandlet (-run=PythonScript) or the Python executor (-ExecutePythonScript) will automatically call ForceEnablePythonAtRuntime as those tools always require Python to be enabled.

A user may choose to enable Python temporarily via the new "Enable Python" menu entry in the Output Log, or permanently via the EnablePythonOverride property in UPythonScriptPluginUserSettings. EnablePythonOverride set to Enabled acts the same as passing -EnablePython on the command line, and set to Disabled acts the same as passing -DisablePython on the command line.

diff --git a/Engine/Plugins/Experimental/PythonScriptPlugin/Source/PythonScriptPlugin/Private/EditorUtilities/EditorPythonExecuter.cpp b/Engine/Plugins/Experimental/PythonScriptPlugin/Source/PythonScriptPlugin/Private/EditorUtilities/EditorPythonExecuter.cpp
--- a/Engine/Plugins/Experimental/PythonScriptPlugin/Source/PythonScriptPlugin/Private/EditorUtilities/EditorPythonExecuter.cpp
+++ b/Engine/Plugins/Experimental/PythonScriptPlugin/Source/PythonScriptPlugin/Private/EditorUtilities/EditorPythonExecuter.cpp
@@ -185,10 +185,13 @@ void FEditorPythonExecuter::OnStartupModule()
 		}
 		else if (IsRunningCommandlet())
 		{
-			UE_LOG(LogEditorPythonExecuter, Error, TEXT("-ExecutePythonScript cannot be used by a commandlet."));
+			UE_LOG(LogEditorPythonExecuter, Error, TEXT("-ExecutePythonScript cannot be used by a commandlet. Use -run=PythonScript instead?"));
 		}
 		else
 		{
+			// If -ExecutePythonScript has been specified then we can assume the user wanted Python support enabled
+			IPythonScriptPlugin::Get()->ForceEnablePythonAtRuntime();
+
 			IPythonScriptPlugin::Get()->RegisterOnPythonConfigured(FSimpleDelegate::CreateLambda([ScriptAndArgs = MoveTemp(ScriptAndArgs), bScriptErrorsAreFatal = FParse::Param(FCommandLine::Get(), TEXT("ScriptErrorsAreFatal"))]() mutable
 			{
 				if (IPythonScriptPlugin::Get()->IsPythonAvailable())
diff --git a/Engine/Plugins/Experimental/PythonScriptPlugin/Source/PythonScriptPlugin/Private/PythonScriptCommandlet.cpp b/Engine/Plugins/Experimental/PythonScriptPlugin/Source/PythonScriptPlugin/Private/PythonScriptCommandlet.cpp
--- a/Engine/Plugins/Experimental/PythonScriptPlugin/Source/PythonScriptPlugin/Private/PythonScriptCommandlet.cpp
+++ b/Engine/Plugins/Experimental/PythonScriptPlugin/Source/PythonScriptPlugin/Private/PythonScriptCommandlet.cpp
@@ -1,7 +1,7 @@
 // Copyright Epic Games, Inc. All Rights Reserved.
 
 #include "PythonScriptCommandlet.h"
-#include "IPythonScriptPlugin.h"
+#include "PythonScriptPlugin.h"
 #include "Logging/LogMacros.h"
 #include "Containers/Ticker.h"
 
@@ -35,14 +35,18 @@ int32 UPythonScriptCommandlet::Main(const FString& Params)
 		return -1;
 	}
 
+#if WITH_PYTHON
+	// We can assume the user wanted Python support enabled when using this commandlet
+	IPythonScriptPlugin::Get()->ForceEnablePythonAtRuntime();
+
 	if (!IPythonScriptPlugin::Get()->IsPythonAvailable())
 	{
 		UE_LOG(LogPythonScriptCommandlet, Error, TEXT("Python script cannot run as Python support is disabled!"));
 		return -1;
 	}
 
-	// Tick once to ensure that any start-up scripts have been run
-	FTSTicker::GetCoreTicker().Tick(0.0f);
+	// Ensure that any start-up scripts have been run
+	FPythonScriptPlugin::Get()->RunStartupScripts();
 
 	{
 		UE_LOG(LogPythonScriptCommandlet, Display, TEXT("Running Python script: %s"), *PythonScript);
@@ -59,4 +63,8 @@ int32 UPythonScriptCommandlet::Main(const FString& Params)
 
 	UE_LOG(LogPythonScriptCommandlet, Display, TEXT("Python script executed successfully"));
 	return 0;
+#else
+	UE_LOG(LogPythonScriptCommandlet, Error, TEXT("Python script cannot run as Python support was disabled at compile-time!"));
+	return -1;
+#endif
 }
diff --git a/Engine/Plugins/Experimental/PythonScriptPlugin/Source/PythonScriptPlugin/Private/PythonScriptLibrary.cpp b/Engine/Plugins/Experimental/PythonScriptPlugin/Source/PythonScriptPlugin/Private/PythonScriptLibrary.cpp
--- a/Engine/Plugins/Experimental/PythonScriptPlugin/Source/PythonScriptPlugin/Private/PythonScriptLibrary.cpp
+++ b/Engine/Plugins/Experimental/PythonScriptPlugin/Source/PythonScriptPlugin/Private/PythonScriptLibrary.cpp
@@ -7,6 +7,8 @@
 #include "PyGenUtil.h"
 #include "PyConversion.h"
 #include "PyWrapperTypeRegistry.h"
+#include "UObject/Script.h"
+#include "UObject/Stack.h"
 #include "UObject/Package.h"
 
 #include UE_INLINE_GENERATED_CPP_BY_NAME(PythonScriptLibrary)
@@ -16,6 +18,35 @@ bool UPythonScriptLibrary::IsPythonAvailable()
 	return IPythonScriptPlugin::Get()->IsPythonAvailable();
 }
 
+bool UPythonScriptLibrary::IsPythonConfigured()
+{
+	return IPythonScriptPlugin::Get()->IsPythonConfigured();
+}
+
+bool UPythonScriptLibrary::IsPythonInitialized()
+{
+	return IPythonScriptPlugin::Get()->IsPythonInitialized();
+}
+
+bool UPythonScriptLibrary::ForceEnablePythonAtRuntime()
+{
+	FString Location;
+#if DO_BLUEPRINT_GUARD
+	if (TConstArrayView<const FFrame*> ScriptStack = FBlueprintContextTracker::Get().GetCurrentScriptStack();
+		ScriptStack.Num() > 0)
+	{
+		TStringBuilder<256> StackDescription;
+		ScriptStack.Last()->GetStackDescription(StackDescription);
+		Location = StackDescription;
+	}
+	else
+#endif	// DO_BLUEPRINT_GUARD
+	{
+		Location = TEXT("<Unknown Blueprint>");
+	}
+	return FPythonScriptPlugin::Get()->ForceEnablePythonAtRuntime(Location);
+}
+
 bool UPythonScriptLibrary::ExecutePythonCommand(const FString& PythonCommand)
 {
 	return IPythonScriptPlugin::Get()->ExecPythonCommand(*PythonCommand);
diff --git a/Engine/Plugins/Experimental/PythonScriptPlugin/Source/PythonScriptPlugin/Private/PythonScriptLibrary.h b/Engine/Plugins/Experimental/PythonScriptPlugin/Source/PythonScriptPlugin/Private/PythonScriptLibrary.h
--- a/Engine/Plugins/Experimental/PythonScriptPlugin/Source/PythonScriptPlugin/Private/PythonScriptLibrary.h
+++ b/Engine/Plugins/Experimental/PythonScriptPlugin/Source/PythonScriptPlugin/Private/PythonScriptLibrary.h
@@ -14,11 +14,32 @@ class UPythonScriptLibrary : public UBlueprintFunctionLibrary
 
 public:
 	/**
-	 * Check to see whether Python support is available in the current build.
+	 * Check to see whether Python support is available in the current environment.
+	 * @note This may return false until IsPythonConfigured is true.
 	 */
 	UFUNCTION(BlueprintCallable, Category = "Python|Execution")
 	static bool IsPythonAvailable();
 
+	/**
+	 * Check to see whether Python has been configured.
+	 * @note Python may be configured but not yet be initialized (@see IsPythonInitialized).
+	 */
+	UFUNCTION(BlueprintCallable, Category = "Python|Execution")
+	static bool IsPythonConfigured();
+
+	/**
+	 * Check to see whether Python has been initialized and is ready to use.
+	 */
+	UFUNCTION(BlueprintCallable, Category = "Python|Execution")
+	static bool IsPythonInitialized();
+
+	/**
+	 * Force Python to be enabled and initialized, regardless of the settings that control its default enabled state.
+	 * @return True if Python was requested to be enabled. Use IsPythonInitialized to verify that it actually initialized.
+	 */
+	UFUNCTION(BlueprintCallable, Category = "Python|Execution")
+	static bool ForceEnablePythonAtRuntime();
+
 	/**
 	 * Execute the given Python command.
 	 * @param PythonCommand The command to run. This may be literal Python code, or a file (with optional arguments) that you want to run.
diff --git a/Engine/Plugins/Experimental/PythonScriptPlugin/Source/PythonScriptPlugin/Private/PythonScriptPlugin.cpp b/Engine/Plugins/Experimental/PythonScriptPlugin/Source/PythonScriptPlugin/Private/PythonScriptPlugin.cpp
--- a/Engine/Plugins/Experimental/PythonScriptPlugin/Source/PythonScriptPlugin/Private/PythonScriptPlugin.cpp
+++ b/Engine/Plugins/Experimental/PythonScriptPlugin/Source/PythonScriptPlugin/Private/PythonScriptPlugin.cpp
@@ -27,6 +27,7 @@
 #include "Misc/Paths.h"
 #include "Misc/FileHelper.h"
 #include "Misc/PackageName.h"
+#include "Misc/SourceLocationUtils.h"
 #include "Misc/ScopedSlowTask.h"
 #include "HAL/FileManager.h"
 #include "HAL/LowLevelMemTracker.h"
@@ -143,6 +144,20 @@ bool IsPythonEnabled()
 		return true;
 	}
 
+	switch (GetDefault<UPythonScriptPluginUserSettings>()->EnablePythonOverride)
+	{
+		case EPythonEnabledOverrideState::Enable:
+			UE_LOG(LogPython, Log, TEXT("Python enabled via PythonScriptPluginUserSettings"));
+			return true;
+
+		case EPythonEnabledOverrideState::Disable:
+			UE_LOG(LogPython, Log, TEXT("Python disabled via PythonScriptPluginUserSettings"));
+			return false;
+
+		default:
+			break;
+	}
+
 	if (UE::Python::Private::bIsEnabledByDefault)
 	{
 		UE_LOG(LogPython, Log, TEXT("Python enabled via CVar 'Engine.Python.IsEnabledByDefault'"));
@@ -686,6 +701,44 @@ bool FPythonScriptPlugin::IsPythonInitialized() const
 #endif
 }
 
+bool FPythonScriptPlugin::ForceEnablePythonAtRuntime(UE::FSourceLocation Location)
+{
+	return ForceEnablePythonAtRuntime(UE::SourceLocation::ToFullString(Location));
+}
+
+bool FPythonScriptPlugin::ForceEnablePythonAtRuntime(const FString& Location)
+{
+#if WITH_PYTHON
+	if (!bIsForceEnabledAtRuntime)
+	{
+		// Promote the log to a warning if called when Python is disabled by default, or was explicitly disabled by the user
+		if (!UE::Python::Private::bIsEnabledByDefault || GetDefault<UPythonScriptPluginUserSettings>()->EnablePythonOverride == EPythonEnabledOverrideState::Disable)
+		{
+			UE_LOG(LogPython, Warning, TEXT("Python enabled via IPythonScriptPlugin::ForceEnablePythonAtRuntime: %s"), *Location);
+		}
+		else
+		{
+			UE_LOG(LogPython, Log, TEXT("Python enabled via IPythonScriptPlugin::ForceEnablePythonAtRuntime: %s"), *Location);
+		}
+		bIsForceEnabledAtRuntime = true;
+
+		if (bIsConfigured && !bIsInterpreterInitialized)
+		{
+			// ForceEnablePythonAtRuntime was called after OnPostEngineInit, so re-run this now that bIsForceEnabledAtRuntime is true
+			ConfigureAndInitializePython();
+			if (bIsInterpreterInitialized && bHasTicked)
+			{
+				// If we've already ticked once, then we can also run the start-up scripts
+				RunStartupScripts();
+			}
+		}
+	}
+	return true;
+#else
+	return false;
+#endif
+}
+
 bool FPythonScriptPlugin::ExecPythonCommand(const TCHAR* InPythonCommand)
 {
 	FPythonCommandEx PythonCommand;
@@ -848,12 +901,54 @@ void FPythonScriptPlugin::StartupModule()
 
 #if WITH_PYTHON
 void FPythonScriptPlugin::OnPostEngineInit()
+{
+	TickOnceHandle = FTSTicker::GetCoreTicker().AddTicker(FTickerDelegate::CreateLambda([this](float DeltaTime)
+	{
+		bHasTicked = true;
+		return false;
+	}));
+
+#if WITH_EDITOR
+	// Register a menu handler to enable Python if it is currently disabled at runtime
+	if (UToolMenus::IsToolMenuUIEnabled())
+	{
+		static const FName MenuName = "OutputLog.ConsoleInputBox.CmdExecMenu";
+		if (UToolMenu* Menu = UToolMenus::Get()->ExtendMenu(MenuName))
+		{
+			Menu->AddDynamicSection("DynamicPythonInit", FNewToolMenuDelegate::CreateLambda([](UToolMenu* InMenu)
+			{
+				FPythonScriptPlugin* This = FPythonScriptPlugin::Get();
+
+				if (This->IsPythonConfigured() && !This->IsPythonAvailable())
+				{
+					FToolMenuSection& Section = InMenu->AddSection("PythonInit");
+					Section.AddMenuEntry(
+						"EnablePython",
+						LOCTEXT("Menu.EnablePython.Label", "Enable Python"),
+						LOCTEXT("Menu.EnablePython.ToolTip", "Enable Python for the current editor session.\nTo enable Python permanently, see 'Editor Preferences -> Python -> Enable Python Override'."),
+						FSlateIcon(),
+						FUIAction(
+							FExecuteAction::CreateLambda([This] { This->ForceEnablePythonAtRuntime(); }),
+							FCanExecuteAction::CreateLambda([] { return true; })
+						),
+						EUserInterfaceActionType::Button
+					);
+				}
+			}));
+		}
+	}
+#endif
+
+	ConfigureAndInitializePython();
+}
+
+void FPythonScriptPlugin::ConfigureAndInitializePython()
 {
 	LLM_SCOPE_BYNAME(TEXT("PythonScriptPlugin"));
 
 	// Determine whether Python is enabled in OnPostEngineInit so that other things have time to override the CVar
 	// IsPythonAvailable() will always return false prior to configuring bIsEnabled
-	bIsEnabled = UE::Python::Private::IsPythonEnabled();
+	bIsEnabled = bIsForceEnabledAtRuntime || UE::Python::Private::IsPythonEnabled();
 	
 	bIsConfigured = true;
 	OnPythonConfiguredDelegate.Broadcast();
@@ -1255,6 +1350,7 @@ void FPythonScriptPlugin::ShutdownPython()
 	FCoreUObjectDelegates::GetPreGarbageCollectDelegate().RemoveAll(this);
 
 	FTSTicker::GetCoreTicker().RemoveTicker(TickHandle);
+	FTSTicker::GetCoreTicker().RemoveTicker(TickOnceHandle);
 	if (ModuleDelayedHandle.IsValid())
 	{
 		FTSTicker::GetCoreTicker().RemoveTicker(ModuleDelayedHandle);
@@ -1304,6 +1400,8 @@ void FPythonScriptPlugin::ShutdownPython()
 	bIsInterpreterInitialized = false;
 	bIsFullyInitialized = false;
 	bIsEnabled = false;
+	bIsForceEnabledAtRuntime = false;
+	bHasTicked = false;
 }
 
 void FPythonScriptPlugin::InitPipInstaller()
@@ -1455,6 +1553,24 @@ void FPythonScriptPlugin::RunStartupScripts()
 		PySysPaths = PyUtil::GetSystemPaths();
 	}
 
+	auto RunTimedStartupScript = [this](const FString& StartupScript, const FText& StartupScriptInfoText, const bool bRunAsFile)
+	{
+		UE_SCOPED_TIMER(*StartupScriptInfoText.ToString(), LogPython, Display);
+		if (bRunAsFile)
+		{
+			// Execute these files in the "public" scope, as if their contents had been run directly in the console
+			// This allows them to be used to set-up an editor environment for the console
+			FPythonCommandEx InitUnrealPythonCommand;
+			InitUnrealPythonCommand.FileExecutionScope = EPythonFileExecutionScope::Public;
+
+			RunFile(*StartupScript, *InitUnrealPythonCommand.Command, InitUnrealPythonCommand);
+		}
+		else
+		{
+			ExecPythonCommand(*StartupScript);
+		}
+	};
+
 	FScopedSlowTask Progress(PySysPaths.Num() + GetDefault<UPythonScriptPluginSettings>()->StartupScripts.Num(), LOCTEXT("PythonScriptPluginInitScripts", "Running Python start-up scripts..."), true);
 	Progress.MakeDialogDelayed(0.1f);
 	for (const FString& PySysPath : PySysPaths)
@@ -1466,13 +1582,7 @@ void FPythonScriptPlugin::RunStartupScripts()
 			Progress.EnterProgressFrame(1.0f, StartupScriptInfoText);
 			Progress.ForceRefresh();
 
-			// Execute these files in the "public" scope, as if their contents had been run directly in the console
-			// This allows them to be used to set-up an editor environment for the console
-			FPythonCommandEx InitUnrealPythonCommand;
-			InitUnrealPythonCommand.FileExecutionScope = EPythonFileExecutionScope::Public;
-			
-			UE_SCOPED_TIMER(*StartupScriptInfoText.ToString(), LogPython, Display);
-			RunFile(*PotentialFilePath, *InitUnrealPythonCommand.Command, InitUnrealPythonCommand);
+			RunTimedStartupScript(PotentialFilePath, StartupScriptInfoText, /*bRunAsFile*/true);
 		}
 		else
 		{
@@ -1485,8 +1595,7 @@ void FPythonScriptPlugin::RunStartupScripts()
 		FText StartupScriptInfoText = FText::Format(LOCTEXT("PythonScriptPluginInitScripts_Running", "Running start-up script {0}..."), FText::FromString(StartupScript));
 		Progress.EnterProgressFrame(1.0f, StartupScriptInfoText);
 
-		UE_SCOPED_TIMER(*StartupScriptInfoText.ToString(), LogPython, Display);
-		ExecPythonCommand(*StartupScript);
+		RunTimedStartupScript(StartupScript, StartupScriptInfoText, /*bRunAsFile*/false);
 	}
 
 	// Notify any external listeners
@@ -1508,6 +1617,8 @@ void FPythonScriptPlugin::RunStartupScripts()
 
 void FPythonScriptPlugin::Tick(const float InDeltaTime)
 {
+	bHasTicked = true;
+
 	RunStartupScripts();
 
 	RemoteExecution->Tick(InDeltaTime);
diff --git a/Engine/Plugins/Experimental/PythonScriptPlugin/Source/PythonScriptPlugin/Private/PythonScriptPlugin.h b/Engine/Plugins/Experimental/PythonScriptPlugin/Source/PythonScriptPlugin/Private/PythonScriptPlugin.h
--- a/Engine/Plugins/Experimental/PythonScriptPlugin/Source/PythonScriptPlugin/Private/PythonScriptPlugin.h
+++ b/Engine/Plugins/Experimental/PythonScriptPlugin/Source/PythonScriptPlugin/Private/PythonScriptPlugin.h
@@ -108,6 +108,8 @@ public:
 	virtual bool IsPythonAvailable() const override;
 	virtual bool IsPythonConfigured() const override;
 	virtual bool IsPythonInitialized() const override;
+	virtual bool ForceEnablePythonAtRuntime(UE::FSourceLocation Location = UE::FSourceLocation::Current()) override;
+	bool ForceEnablePythonAtRuntime(const FString& Location);
 	virtual bool ExecPythonCommand(const TCHAR* InPythonCommand) override;
 	virtual bool ExecPythonCommandEx(FPythonCommandEx& InOutPythonCommand) override;
 	virtual FString GetInterpreterExecutablePath() const override;
@@ -126,8 +128,6 @@ public:
 #if WITH_PYTHON
 	void RunStartupScripts();
 
-	void OnPostEngineInit();
-
 	/** Sync the remote execution environment to the current settings, starting or stopping it as required */
 	void SyncRemoteExecutionToSettings();
 
@@ -159,6 +159,10 @@ protected:
 
 private:
 #if WITH_PYTHON
+	void OnPostEngineInit();
+
+	void ConfigureAndInitializePython();
+
 	void InitializePython();
 
 	void ShutdownPython();
@@ -212,6 +216,7 @@ private:
 	FPythonREPLCommandExecutor CmdREPLExec;
 	IPythonCommandMenu* CmdMenu = nullptr;
 	FTSTicker::FDelegateHandle TickHandle;
+	FTSTicker::FDelegateHandle TickOnceHandle;
 	FTSTicker::FDelegateHandle ModuleDelayedHandle;
 
 	PyUtil::FPyApiBuffer PyProgramName;
@@ -226,6 +231,8 @@ private:
 	bool bIsInterpreterInitialized = false; // Python interpreter is available; may not yet have run startup scripts
 	bool bIsFullyInitialized = false; // Python interpreter is available, and start-up scripts have run
 	bool bIsEnabled = false;
+	bool bIsForceEnabledAtRuntime = false;
+	bool bHasTicked = false;
 
 #if PY_MAJOR_VERSION >= 3 && PY_MINOR_VERSION >= 11
 	PyConfig ModulePyConfig;
diff --git a/Engine/Plugins/Experimental/PythonScriptPlugin/Source/PythonScriptPlugin/Private/PythonScriptPluginSettings.h b/Engine/Plugins/Experimental/PythonScriptPlugin/Source/PythonScriptPlugin/Private/PythonScriptPluginSettings.h
--- a/Engine/Plugins/Experimental/PythonScriptPlugin/Source/PythonScriptPlugin/Private/PythonScriptPluginSettings.h
+++ b/Engine/Plugins/Experimental/PythonScriptPlugin/Source/PythonScriptPlugin/Private/PythonScriptPluginSettings.h
@@ -29,6 +29,19 @@ enum class ETypeHintingMode : uint8
 	TypeChecker,
 };
 
+UENUM()
+enum class EPythonEnabledOverrideState : uint8
+{
+	/** Python will be enabled based on the current project settings and any command line overrides. */
+	None,
+
+	/** Enable Python, even if the project has it disabled by default (behaves like '-EnablePython' on the command line). Note: Does not take precedence over '-DisablePython' on the command line. */
+	Enable,
+
+	/** Disable Python, even if the project has it enabled by default (behaves like '-DisablePython' on the command line). Note: Does not take precedence over '-EnablePython' on the command line. */
+	Disable,
+};
+
 /**
  * Configure the Python plug-in.
  */
@@ -150,6 +163,12 @@ public:
 	virtual FText GetSectionText() const override;
 #endif
 
+	/**
+	 * Should we override the default enabled state for Python for this project?
+	 */
+	UPROPERTY(config, EditAnywhere, Category=Python, meta=(ConfigRestartRequired=true))
+	EPythonEnabledOverrideState EnablePythonOverride = EPythonEnabledOverrideState::None;
+
 	/**
 	 * Should Developer Mode be enabled on the Python interpreter?
 	 *
diff --git a/Engine/Plugins/Experimental/PythonScriptPlugin/Source/PythonScriptPlugin/Public/IPythonScriptPlugin.h b/Engine/Plugins/Experimental/PythonScriptPlugin/Source/PythonScriptPlugin/Public/IPythonScriptPlugin.h
--- a/Engine/Plugins/Experimental/PythonScriptPlugin/Source/PythonScriptPlugin/Public/IPythonScriptPlugin.h
+++ b/Engine/Plugins/Experimental/PythonScriptPlugin/Source/PythonScriptPlugin/Public/IPythonScriptPlugin.h
@@ -4,6 +4,7 @@
 
 #include "CoreTypes.h"
 #include "PythonScriptTypes.h"
+#include "Misc/SourceLocation.h"
 #include "Modules/ModuleManager.h"
 #include "Modules/ModuleInterface.h"
 
@@ -34,6 +35,12 @@ public:
 	 */
 	virtual bool IsPythonInitialized() const = 0;
 
+	/**
+	 * Force Python to be enabled and initialized, regardless of the settings that control its default enabled state.
+	 * @return True if Python was requested to be enabled. Use IsPythonInitialized to verify that it actually initialized.
+	 */
+	virtual bool ForceEnablePythonAtRuntime(UE::FSourceLocation Location = UE::FSourceLocation::Current()) = 0;
+
 	/**
 	 * Execute the given Python command.
 	 * This may be literal Python code, or a file (with optional arguments) that you want to run.
