Fixed ScriptName deprecation failing to handle the real property/function name being used as the deprecated name

eg) A function called GetLODs in C++ having a ScriptName of "GetLods;GetLODs" would fail to export the "GetLODs" deprecated entry as it still resolved to an existing function, despite having a different name than "GetLods" when exposed to Python (GetLods -> get_lods; GetLODs -> get_lo_ds).

diff --git a/Engine/Plugins/Experimental/PythonScriptPlugin/Source/PythonScriptPlugin/Private/PyGenUtil.cpp b/Engine/Plugins/Experimental/PythonScriptPlugin/Source/PythonScriptPlugin/Private/PyGenUtil.cpp
--- a/Engine/Plugins/Experimental/PythonScriptPlugin/Source/PythonScriptPlugin/Private/PyGenUtil.cpp
+++ b/Engine/Plugins/Experimental/PythonScriptPlugin/Source/PythonScriptPlugin/Private/PyGenUtil.cpp
@@ -3054,19 +3054,24 @@ TArray<TTuple<FSoftObjectPath, FString>> GetDeprecatedFunctionPythonNames(const
 	check(FuncOwner);
 
 	TArray<TTuple<FSoftObjectPath, FString>> FuncNames = GetDeprecatedFieldPythonNamesImpl(InFunc, ScriptNameMetaDataKey);
-	for (auto FuncNamesIt = FuncNames.CreateIterator(); FuncNamesIt; ++FuncNamesIt)
+	if (FuncNames.Num() > 0)
 	{
-		FString& FuncName = FuncNamesIt->Value;
-
-		// Remove any deprecated names that clash with an existing Python exposed function
-		const UFunction* DeprecatedFunc = FuncOwner->FindFunctionByName(*FuncName);
-		if (DeprecatedFunc && ShouldExportFunction(DeprecatedFunc))
+		const FString FuncPythonName = GetFunctionPythonName(InFunc);
+		for (auto FuncNamesIt = FuncNames.CreateIterator(); FuncNamesIt; ++FuncNamesIt)
 		{
-			FuncNamesIt.RemoveCurrent();
-			continue;
-		}
+			FString& FuncName = FuncNamesIt->Value;
 
-		FuncName = PythonizeName(FuncName, EPythonizeNameCase::Lower);
+			// Remove any deprecated names that clash with an existing Python exposed function, 
+			// unless that function is ourself, and the Pythonized name is different
+			// This can happen when fixing case, eg LODs -> Lods producing lods rather than lo_ds
+			const UFunction* DeprecatedFunc = FuncOwner->FindFunctionByName(*FuncName);
+			FuncName = PythonizeName(FuncName, EPythonizeNameCase::Lower);
+			if (DeprecatedFunc && ShouldExportFunction(DeprecatedFunc) && (DeprecatedFunc != InFunc || FuncName.Equals(FuncPythonName)))
+			{
+				FuncNamesIt.RemoveCurrent();
+				continue;
+			}
+		}
 	}
 
 	return FuncNames;
@@ -3134,19 +3139,24 @@ TArray<TTuple<FSoftObjectPath, FString>> GetDeprecatedPropertyPythonNames(const
 	check(PropOwner);
 
 	TArray<TTuple<FSoftObjectPath, FString>> PropNames = GetDeprecatedFieldPythonNamesImpl(InProp, ScriptNameMetaDataKey);
-	for (auto PropNamesIt = PropNames.CreateIterator(); PropNamesIt; ++PropNamesIt)
-	{
-		FString& PropName = PropNamesIt->Value;
-
-		// Remove any deprecated names that clash with an existing Python exposed property
-		const FProperty* DeprecatedProp = PropOwner->FindPropertyByName(*PropName);
-		if (DeprecatedProp && ShouldExportProperty(DeprecatedProp))
+	if (PropNames.Num() > 0)
 	{
-			PropNamesIt.RemoveCurrent();
-			continue;
-	}
+		const FString PropPythonName = GetPropertyPythonName(InProp);
+		for (auto PropNamesIt = PropNames.CreateIterator(); PropNamesIt; ++PropNamesIt)
+		{
+			FString& PropName = PropNamesIt->Value;
 
-		PropName = PythonizeName(PropName, EPythonizeNameCase::Lower);
+			// Remove any deprecated names that clash with an existing Python exposed property, 
+			// unless that property is ourself, and the Pythonized name is different
+			// This can happen when fixing case, eg LODs -> Lods producing lods rather than lo_ds
+			const FProperty* DeprecatedProp = PropOwner->FindPropertyByName(*PropName);
+			PropName = PythonizeName(PropName, EPythonizeNameCase::Lower);
+			if (DeprecatedProp && ShouldExportProperty(DeprecatedProp) && (DeprecatedProp != InProp || PropName.Equals(PropPythonName)))
+			{
+				PropNamesIt.RemoveCurrent();
+				continue;
+			}
+		}
 	}
 
 	return PropNames;
