Defer any -execcmds for Python that are requested before it's fully initialized

diff --git a/Engine/Plugins/Experimental/PythonScriptPlugin/Source/PythonScriptPlugin/Private/PythonScriptPlugin.cpp b/Engine/Plugins/Experimental/PythonScriptPlugin/Source/PythonScriptPlugin/Private/PythonScriptPlugin.cpp
--- a/Engine/Plugins/Experimental/PythonScriptPlugin/Source/PythonScriptPlugin/Private/PythonScriptPlugin.cpp
+++ b/Engine/Plugins/Experimental/PythonScriptPlugin/Source/PythonScriptPlugin/Private/PythonScriptPlugin.cpp
@@ -1023,7 +1023,14 @@ bool FPythonScriptPlugin::Exec_Runtime(UWorld* InWorld, const TCHAR* Cmd, FOutpu
 #if WITH_PYTHON
 	if (FParse::Command(&Cmd, TEXT("PY")))
 	{
-		ExecPythonCommand(Cmd);
+		if (IsPythonInitialized())
+		{
+			ExecPythonCommand(Cmd);
+		}
+		else
+		{
+			DeferredCommands.Add(Cmd);
+		}
 		return true;
 	}
 #endif	// WITH_PYTHON
@@ -1598,6 +1605,13 @@ void FPythonScriptPlugin::RunStartupScripts()
 		RunTimedStartupScript(StartupScript, StartupScriptInfoText, /*bRunAsFile*/false);
 	}
 
+	// Run any deferred commands now
+	for (const FString& DeferredCommand : DeferredCommands)
+	{
+		ExecPythonCommand(*DeferredCommand);
+	}
+	DeferredCommands.Reset();
+
 	// Notify any external listeners
 	OnPythonInitializedDelegate.Broadcast();
 	OnPythonInitializedDelegate.Clear();
diff --git a/Engine/Plugins/Experimental/PythonScriptPlugin/Source/PythonScriptPlugin/Private/PythonScriptPlugin.h b/Engine/Plugins/Experimental/PythonScriptPlugin/Source/PythonScriptPlugin/Private/PythonScriptPlugin.h
--- a/Engine/Plugins/Experimental/PythonScriptPlugin/Source/PythonScriptPlugin/Private/PythonScriptPlugin.h
+++ b/Engine/Plugins/Experimental/PythonScriptPlugin/Source/PythonScriptPlugin/Private/PythonScriptPlugin.h
@@ -218,6 +218,7 @@ private:
 	FTSTicker::FDelegateHandle TickHandle;
 	FTSTicker::FDelegateHandle TickOnceHandle;
 	FTSTicker::FDelegateHandle ModuleDelayedHandle;
+	TArray<FString> DeferredCommands;
 
 	PyUtil::FPyApiBuffer PyProgramName;
 	PyUtil::FPyApiBuffer PyHomePath;
