[Python]
- GenerateWrappers can now run in parallel
- 8.7s to 1.5s during startup
- Off by default for now, can be activated by using -dpcvars=Engine.Python.ForceWrapperGenerationOnSingleThread=false

diff --git a/Engine/Plugins/Experimental/PythonScriptPlugin/Source/PythonScriptPlugin/Private/PyGenUtil.cpp b/Engine/Plugins/Experimental/PythonScriptPlugin/Source/PythonScriptPlugin/Private/PyGenUtil.cpp
--- a/Engine/Plugins/Experimental/PythonScriptPlugin/Source/PythonScriptPlugin/Private/PyGenUtil.cpp
+++ b/Engine/Plugins/Experimental/PythonScriptPlugin/Source/PythonScriptPlugin/Private/PyGenUtil.cpp
@@ -81,7 +81,7 @@ const FName NativeMakeFuncMetaDataKey = TEXT("NativeMakeFunc");
 const FName ReturnValueKey = TEXT("ReturnValue");
 const TCHAR* HiddenMetaDataKey = TEXT("Hidden");
 
-TSharedPtr<IBreakIterator> NameBreakIterator;
+static thread_local TSharedPtr<IBreakIterator> NameBreakIterator;
 static ETypeHintingMode GTypeHintingMode = ETypeHintingMode::Off; // Use through SetTypeHintingMode()/IsTypeHintingEnabled()
 
 
@@ -528,20 +528,21 @@ void FGeneratedWrappedPropertyDoc::AppendDocString(const TArray<FGeneratedWrappe
 		return;
 	}
 
+	TStringBuilder<4096> Buffer;
 	if (!OutStr.IsEmpty())
 	{
 		if (OutStr[OutStr.Len() - 1] != TEXT('\n'))
 		{
-			OutStr += LINE_TERMINATOR;
+			Buffer += LINE_TERMINATOR;
 		}
 	}
 
-	OutStr += LINE_TERMINATOR TEXT("**Editor Properties:** (see get_editor_property/set_editor_property)") LINE_TERMINATOR;
+	Buffer += LINE_TERMINATOR TEXT("**Editor Properties:** (see get_editor_property/set_editor_property)") LINE_TERMINATOR;
 	for (const FGeneratedWrappedPropertyDoc& Doc : InDocs)
 	{
-		OutStr += LINE_TERMINATOR TEXT("- ``");  // add as a list and code style
-		OutStr += Doc.PythonPropName;
-		OutStr += TEXT("`` ");
+		Buffer += LINE_TERMINATOR TEXT("- ``");  // add as a list and code style
+		Buffer += Doc.PythonPropName;
+		Buffer += TEXT("`` ");
 
 		bool bMultipleLines = false;
 
@@ -552,7 +553,7 @@ void FGeneratedWrappedPropertyDoc::AppendDocString(const TArray<FGeneratedWrappe
 		{
 			if (bMultipleLines)
 			{
-				OutStr += LINE_TERMINATOR TEXT("  ");
+				Buffer += LINE_TERMINATOR TEXT("  ");
 			}
 			bMultipleLines = true;
 
@@ -562,7 +563,7 @@ void FGeneratedWrappedPropertyDoc::AppendDocString(const TArray<FGeneratedWrappe
 				++EditorDocStringParseIndex;
 			}
 
-			OutStr += EditorDocString.Mid(LineStartIndex, EditorDocStringParseIndex - LineStartIndex);
+			Buffer += EditorDocString.Mid(LineStartIndex, EditorDocStringParseIndex - LineStartIndex);
 
 			if (EditorDocStringParseIndex + 1 < EditorDocString.Len() && EditorDocString[EditorDocStringParseIndex] == TEXT('\r') && EditorDocString[EditorDocStringParseIndex + 1] == TEXT('\n'))
 			{
@@ -574,6 +575,8 @@ void FGeneratedWrappedPropertyDoc::AppendDocString(const TArray<FGeneratedWrappe
 			}
 		}
 	}
+
+	OutStr += Buffer.ToString();
 }
 
 
@@ -650,13 +653,15 @@ bool FGeneratedWrappedType::Finalize()
 			bSuccess = true;
 		}
 		FinalizedState = EFinalizedState::Finalized;
-	}
 
-	if (bSuccess)
-	{
-		Finalize_PostReady();
-		FPyWrapperBaseMetaData::SetMetaData(&PyType, MetaData.Get());
-		return true;
+		// Run these under the same GIL as above since we're going to need it in both functions below
+		// and it will avoid 2 more locks in rapid successions.
+		if (bSuccess)
+		{
+			Finalize_PostReady();
+			FPyWrapperBaseMetaData::SetMetaData(&PyType, MetaData.Get());
+			return true;
+		}
 	}
 
 	return false;
diff --git a/Engine/Plugins/Experimental/PythonScriptPlugin/Source/PythonScriptPlugin/Private/PyGenUtil.h b/Engine/Plugins/Experimental/PythonScriptPlugin/Source/PythonScriptPlugin/Private/PyGenUtil.h
--- a/Engine/Plugins/Experimental/PythonScriptPlugin/Source/PythonScriptPlugin/Private/PyGenUtil.h
+++ b/Engine/Plugins/Experimental/PythonScriptPlugin/Source/PythonScriptPlugin/Private/PyGenUtil.h
@@ -688,9 +688,9 @@ namespace PyGenUtil
 
 		virtual ~FGeneratedWrappedType() = default;
 
-		FGeneratedWrappedType(FGeneratedWrappedType&&) = default;
+		FGeneratedWrappedType(FGeneratedWrappedType&&) = delete;
 		FGeneratedWrappedType(const FGeneratedWrappedType&) = delete;
-		FGeneratedWrappedType& operator=(FGeneratedWrappedType&&) = default;
+		FGeneratedWrappedType& operator=(FGeneratedWrappedType&&) = delete;
 		FGeneratedWrappedType& operator=(const FGeneratedWrappedType&) = delete;
 
 		/** Called to ready the generated type with Python */
@@ -725,15 +725,53 @@ namespace PyGenUtil
 
 		/** What is the current finalization state of this generated type? */
 		EFinalizedState FinalizedState = EFinalizedState::Initial;
+
+		/* Per-object lock used during the construction phase */
+		struct FDebuggableMutex
+		{
+			// We should never go above 1 lock per stack because
+			// going recursive on those lock will create a very deadlock
+			// prone situation when run under multiple thread that will
+			// be hard to repro depending on the order in which things
+			// are recursed into.
+			static int32& GetThreadLocal()
+			{
+				static thread_local int32 StackLockCount = 0;
+				return StackLockCount;
+			}
+
+			UE::FMutex Mutex;
+			uint32     OwnerThreadId = 0;
+
+			void Lock()
+			{
+				check(GetThreadLocal() == 0);
+				Mutex.Lock();
+				GetThreadLocal()++;
+				check(OwnerThreadId == 0);
+				OwnerThreadId = FPlatformTLS::GetCurrentThreadId();
+			}
+
+			void Unlock()
+			{
+				OwnerThreadId = 0;
+				check(GetThreadLocal() == 1);
+				GetThreadLocal()--;
+				Mutex.Unlock();
+			}
+		} Lock;
+
+		/* Used to atomically publish readyness of this object */
+		std::atomic<bool> bIsReadyForUse;
 	};
 
 	/** Stores the data needed by a runtime generated Python struct type */
 	struct FGeneratedWrappedStructType : public FGeneratedWrappedType, public TGeneratedWrappedDynamicMethodsMixin<FGeneratedWrappedStructType>, public TGeneratedWrappedDynamicConstantsMixin<FGeneratedWrappedStructType>
 	{
 		FGeneratedWrappedStructType() = default;
-		FGeneratedWrappedStructType(FGeneratedWrappedStructType&&) = default;
+		FGeneratedWrappedStructType(FGeneratedWrappedStructType&&) = delete;
 		FGeneratedWrappedStructType(const FGeneratedWrappedStructType&) = delete;
-		FGeneratedWrappedStructType& operator=(FGeneratedWrappedStructType&&) = default;
+		FGeneratedWrappedStructType& operator=(FGeneratedWrappedStructType&&) = delete;
 		FGeneratedWrappedStructType& operator=(const FGeneratedWrappedStructType&) = delete;
 
 		/** Internal version of Finalize, called before readying the type with Python */
@@ -753,9 +791,9 @@ namespace PyGenUtil
 	struct FGeneratedWrappedClassType : public FGeneratedWrappedType, public TGeneratedWrappedDynamicMethodsMixin<FGeneratedWrappedStructType>, public TGeneratedWrappedDynamicConstantsMixin<FGeneratedWrappedStructType>
 	{
 		FGeneratedWrappedClassType() = default;
-		FGeneratedWrappedClassType(FGeneratedWrappedClassType&&) = default;
+		FGeneratedWrappedClassType(FGeneratedWrappedClassType&&) = delete;
 		FGeneratedWrappedClassType(const FGeneratedWrappedClassType&) = delete;
-		FGeneratedWrappedClassType& operator=(FGeneratedWrappedClassType&&) = default;
+		FGeneratedWrappedClassType& operator=(FGeneratedWrappedClassType&&) = delete;
 		FGeneratedWrappedClassType& operator=(const FGeneratedWrappedClassType&) = delete;
 
 		/** Internal version of Finalize, called before readying the type with Python */
@@ -784,9 +822,9 @@ namespace PyGenUtil
 	struct FGeneratedWrappedEnumEntry
 	{
 		FGeneratedWrappedEnumEntry() = default;
-		FGeneratedWrappedEnumEntry(FGeneratedWrappedEnumEntry&&) = default;
+		FGeneratedWrappedEnumEntry(FGeneratedWrappedEnumEntry&&) = delete;
 		FGeneratedWrappedEnumEntry(const FGeneratedWrappedEnumEntry&) = default;
-		FGeneratedWrappedEnumEntry& operator=(FGeneratedWrappedEnumEntry&&) = default;
+		FGeneratedWrappedEnumEntry& operator=(FGeneratedWrappedEnumEntry&&) = delete;
 		FGeneratedWrappedEnumEntry& operator=(const FGeneratedWrappedEnumEntry&) = default;
 
 		/** The name of the entry */
@@ -803,9 +841,9 @@ namespace PyGenUtil
 	struct FGeneratedWrappedEnumType : public FGeneratedWrappedType
 	{
 		FGeneratedWrappedEnumType() = default;
-		FGeneratedWrappedEnumType(FGeneratedWrappedEnumType&&) = default;
+		FGeneratedWrappedEnumType(FGeneratedWrappedEnumType&&) = delete;
 		FGeneratedWrappedEnumType(const FGeneratedWrappedEnumType&) = delete;
-		FGeneratedWrappedEnumType& operator=(FGeneratedWrappedEnumType&&) = default;
+		FGeneratedWrappedEnumType& operator=(FGeneratedWrappedEnumType&&) = delete;
 		FGeneratedWrappedEnumType& operator=(const FGeneratedWrappedEnumType&) = delete;
 
 		/** Internal version of Finalize, called after readying the type with Python */
diff --git a/Engine/Plugins/Experimental/PythonScriptPlugin/Source/PythonScriptPlugin/Private/PyWrapperTypeRegistry.cpp b/Engine/Plugins/Experimental/PythonScriptPlugin/Source/PythonScriptPlugin/Private/PyWrapperTypeRegistry.cpp
--- a/Engine/Plugins/Experimental/PythonScriptPlugin/Source/PythonScriptPlugin/Private/PyWrapperTypeRegistry.cpp
+++ b/Engine/Plugins/Experimental/PythonScriptPlugin/Source/PythonScriptPlugin/Private/PyWrapperTypeRegistry.cpp
@@ -22,10 +22,12 @@
 #include "PythonScriptPluginSettings.h"
 #include "ProfilingDebugging/ScopedTimers.h"
 #include "ProfilingDebugging/CpuProfilerTrace.h"
+#include "HAL/IConsoleManager.h"
 #include "HAL/ThreadHeartBeat.h"
 #include "Misc/PackageName.h"
 #include "Misc/Paths.h"
 #include "Misc/FileHelper.h"
+#include "Misc/Optional.h"
 #include "Misc/ScopedSlowTask.h"
 #include "Misc/StringBuilder.h"
 #include "SourceCodeNavigation.h"
@@ -37,6 +39,9 @@
 #include "UObject/StructOnScope.h"
 #include "UObject/NameTypes.h"
 #include "Misc/CoreDelegates.h"
+#include "Async/UniqueLock.h"
+#include "Async/ParallelFor.h"
+#include "Sanitizer/RaceDetector.h"
 #if WITH_EDITOR
 #include "Kismet2/ReloadUtilities.h"
 #endif
@@ -45,6 +50,33 @@
 
 #if WITH_PYTHON
 
+namespace UE::Python::Private
+{
+
+static bool GPythonForceWrapperGenerationOnSingleThread = true;
+static FAutoConsoleVariableRef CVarGPythonForceWrapperGenerationOnSingleThread(
+	TEXT("Engine.Python.ForceWrapperGenerationOnSingleThread"),
+	GPythonForceWrapperGenerationOnSingleThread,
+	TEXT("Force all wrapper generation to happen on a single thread (slow)."),
+	ECVF_Default
+);
+
+#if USING_INSTRUMENTATION
+
+static bool GPythonWrapperGenerationRaceDetector = false;
+static FAutoConsoleVariableRef CVarGPythonWrapperGenerationRaceDetector(
+	TEXT("Engine.Python.WrapperGenerationRaceDetector"),
+	GPythonWrapperGenerationRaceDetector,
+	TEXT("Activate race detector for the duration of the python wrapper generation."),
+	ECVF_Default
+);
+
+#endif
+
+static bool GIsInMultithreadedGeneration = false;
+
+} // namespace UE::Python::Private
+
 namespace UE::Python
 {
 
@@ -826,14 +858,14 @@ FPyWrapperTypeRegistry::FPyWrapperTypeRegistry()
 	// When everything is loaded, warn the user if there are still unresolved Make/Break functions.
 	FCoreDelegates::OnPostEngineInit.AddLambda([this]()
 	{
-		for (const TPair<FString, TSharedPtr<FPyWrapperStructMetaData>>& UnresolvedMakePair : UnresolvedMakeFuncs)
+		UnresolvedMakeFuncs.ForEach([](const TPair<FString, TSharedPtr<FPyWrapperStructMetaData>>& UnresolvedMakePair)
 		{
 			REPORT_PYTHON_GENERATION_ISSUE(Error, TEXT("Struct '%s' is marked as '%s' but the function '%s' could not be found."), *UnresolvedMakePair.Value->Struct->GetPathName(), *PyGenUtil::HasNativeMakeMetaDataKey.ToString(), *UnresolvedMakePair.Key);
-		}
-		for (const TPair<FString, TSharedPtr<FPyWrapperStructMetaData>>& UnresolvedBreakPair : UnresolvedBreakFuncs)
+		});
+		UnresolvedBreakFuncs.ForEach([](const TPair<FString, TSharedPtr<FPyWrapperStructMetaData>>& UnresolvedBreakPair)
 		{
 			REPORT_PYTHON_GENERATION_ISSUE(Error, TEXT("Struct '%s' is marked as '%s' but the function '%s' could not be found."), *UnresolvedBreakPair.Value->Struct->GetPathName(), *PyGenUtil::HasNativeBreakMetaDataKey.ToString(), *UnresolvedBreakPair.Key);
-		}
+		});
 	});
 }
 
@@ -847,13 +879,13 @@ void FPyWrapperTypeRegistry::AddReferencedObjects(FReferenceCollector& InCollect
 {
 	auto AddTypeReferencedObjects = [&InCollector](const auto& PythonWrappedTypes)
 	{
-		for (const auto& PythonWrappedTypePair : PythonWrappedTypes)
+		PythonWrappedTypes.ForEach([&InCollector](const auto& PythonWrappedTypePair)
 		{
 			if (FPyWrapperBaseMetaData* PythonWrappedTypeMetaData = FPyWrapperBaseMetaData::GetMetaData(PythonWrappedTypePair.Value))
 			{
 				PythonWrappedTypeMetaData->AddTypeReferencedObjects(InCollector);
 			}
-		}
+		});
 	};
 
 	AddTypeReferencedObjects(PythonWrappedClasses);
@@ -874,7 +906,7 @@ void FPyWrapperTypeRegistry::RegisterNativePythonModule(PyGenUtil::FNativePython
 
 void FPyWrapperTypeRegistry::RegisterInlineStructFactory(const TSharedRef<const IPyWrapperInlineStructFactory>& InFactory)
 {
-	check(bCanRegisterInlineStructFactories);
+	check(bCanRegisterInlineStructFactories.load(std::memory_order_relaxed));
 	InlineStructFactories.Add(InFactory->GetStructName(), InFactory);
 }
 
@@ -883,7 +915,7 @@ const IPyWrapperInlineStructFactory* FPyWrapperTypeRegistry::GetInlineStructFact
 	return InlineStructFactories.FindRef(StructName).Get();
 }
 
-void FPyWrapperTypeRegistry::GenerateWrappedTypes()
+void FPyWrapperTypeRegistry::GenerateWrappedTypes(bool bAllowMultithreadedGeneration)
 {
 	TRACE_CPUPROFILER_EVENT_SCOPE(FPyWrapperTypeRegistry::GenerateWrappedTypes)
 
@@ -902,19 +934,83 @@ void FPyWrapperTypeRegistry::GenerateWrappedTypes()
 		FScopedSlowTask SlowTask(ObjectsToProcess.Num(), LOCTEXT("GeneratingPythonTypes", "Generating Python Types..."));
 		SlowTask.Visibility = ESlowTaskVisibility::Important; // this function can be very slow, users will benefit from our messages
 
-		FNameBuilder ObjectPackageName;
-		for (UObject* ObjectToProcess : ObjectsToProcess)
+#if USING_INSTRUMENTATION
+		bool bRaceDetectionActivated = false;
+		if (UE::Python::Private::GPythonWrapperGenerationRaceDetector)
 		{
-			SlowTask.EnterProgressFrame();
-
-			ObjectToProcess->GetPackage()->GetFName().ToString(ObjectPackageName);
-			if (FPackageName::IsScriptPackage(ObjectPackageName.ToView()))
+			UE::Sanitizer::RaceDetector::ToggleRaceDetection(true);
+			bRaceDetectionActivated = true;
+		}
+		ON_SCOPE_EXIT
+		{
+			if (bRaceDetectionActivated)
 			{
-				GenerateWrappedTypeForObject(ObjectToProcess, GeneratedWrappedTypeReferences, DirtyModules);
+				UE::Sanitizer::RaceDetector::ToggleRaceDetection(false);
 			}
+		};
+#endif
+
+		// Release the GIL if we want other threads to be able to do some work.
+		PyThreadState* ThreadState = PyEval_SaveThread();
+
+		UE::FMutex Mutex;
+		std::atomic<int32> SharedIndex = 0;
+		
+		const bool bIsMultithreadedGenerationAllowed = !UE::Python::Private::GPythonForceWrapperGenerationOnSingleThread && bAllowMultithreadedGeneration;
+		TGuardValue GuardIsInMultithreadedGeneration(UE::Python::Private::GIsInMultithreadedGeneration, bIsMultithreadedGenerationAllowed);
+
+		EParallelForFlags ParallelForFlags = EParallelForFlags::Unbalanced;
+		if (!bIsMultithreadedGenerationAllowed)
+		{
+			ParallelForFlags |= EParallelForFlags::ForceSingleThread;
 		}
 
-		GenerateWrappedTypesForReferences(GeneratedWrappedTypeReferences, DirtyModules);
+		// Do not try to call EnterProgressFrame on the SlowTask during this because a single FlushRenderingCommand
+		// can take more time than the whole function now that it's parallel.
+		ParallelFor(
+			// No more gains past a certain point because of highly contended PyGIL.
+			FMath::Min(24u, LowLevelTasks::FScheduler::Get().GetNumWorkers()),
+			[this, &Mutex, &SharedIndex, &ObjectsToProcess, &DirtyModules](int32)
+			{
+				TSet<FName> LocalDirtyModules;
+				FGeneratedWrappedTypeReferences LocalGeneratedWrappedTypeReferences;
+
+				while (true)
+				{
+					int32 Index = SharedIndex.fetch_add(1, std::memory_order_relaxed);
+					if (Index >= ObjectsToProcess.Num())
+					{
+						break;
+					}
+
+					UObject* Object = ObjectsToProcess[Index];
+
+					FNameBuilder ObjectPackageName;
+					Object->GetPackage()->GetFName().ToString(ObjectPackageName);
+					if (FPackageName::IsScriptPackage(ObjectPackageName.ToView()))
+					{
+						GenerateWrappedTypeForObject(Object, LocalGeneratedWrappedTypeReferences, LocalDirtyModules);
+					}
+
+					if (IsInGameThread())
+					{
+						ExecutePostInitOnGameThread();
+					}
+				}
+
+				GenerateWrappedTypesForReferences(LocalGeneratedWrappedTypeReferences, LocalDirtyModules);
+
+				UE::TUniqueLock ScopeLock(Mutex);
+				DirtyModules.Append(LocalDirtyModules);
+			},
+			ParallelForFlags
+		);
+
+		PyEval_RestoreThread(ThreadState);
+
+		ExecutePostInitOnGameThread();
+
+		check(RegisteredPostInitForGameThread.IsEmpty());
 	}
 
 	NotifyModulesDirtied(DirtyModules);
@@ -922,6 +1018,20 @@ void FPyWrapperTypeRegistry::GenerateWrappedTypes()
 	UE_LOG(LogPython, Verbose, TEXT("Took %f seconds to generate and initialize Python wrapped types for the initial load."), GenerateDuration);
 }
 
+void FPyWrapperTypeRegistry::ExecutePostInitOnGameThread()
+{
+	TRACE_CPUPROFILER_EVENT_SCOPE(FPyWrapperTypeRegistry::ExecutePostInitOnGameThread);
+
+	check(IsInGameThread());
+	
+	RegisteredPostInitForGameThread.ConsumeAllLifo(
+		[](TFunction<void()> Function)
+		{
+			Function();
+		}
+	);
+}
+
 void FPyWrapperTypeRegistry::GenerateWrappedTypesForModule(const FName ModuleName)
 {
 	TRACE_CPUPROFILER_EVENT_SCOPE(FPyWrapperTypeRegistry::GenerateWrappedTypesForModule)
@@ -962,22 +1072,27 @@ void FPyWrapperTypeRegistry::GenerateWrappedTypesForModule(const FName ModuleNam
 
 	// Try resolving missing make/break native functions. Some USTRUCT uses make/break functions implemented in other C++ module(s)
 	// that might not be loaded yet when Python glue was generated for the struct.
-	for (TMap<FString, TSharedPtr<FPyWrapperStructMetaData>>::TIterator It = UnresolvedMakeFuncs.CreateIterator(); It; ++It)
+	UnresolvedMakeFuncs.RemoveIf([](TPair<FString, TSharedPtr<FPyWrapperStructMetaData>>& Pair)
 	{
-		if (UFunction* MakeFunc = FindObject<UFunction>(nullptr, *It->Key, true))
+		if (UFunction* MakeFunc = FindObject<UFunction>(nullptr, *Pair.Key, true))
 		{
-			UE::Python::SetMakeFunction(*It->Value, MakeFunc);
-			It.RemoveCurrent();
+			UE::Python::SetMakeFunction(*Pair.Value, MakeFunc);
+			return true; // Remove
 		}
-	}
-	for (TMap<FString, TSharedPtr<FPyWrapperStructMetaData>>::TIterator It = UnresolvedBreakFuncs.CreateIterator(); It; ++It)
+
+		return false;
+	});
+	
+	UnresolvedBreakFuncs.RemoveIf([](TPair<FString, TSharedPtr<FPyWrapperStructMetaData>>& Pair)
 	{
-		if (UFunction* BreakFunc = FindObject<UFunction>(nullptr, *It->Key, true))
+		if (UFunction* BreakFunc = FindObject<UFunction>(nullptr, *Pair.Key, true))
 		{
-			UE::Python::SetBreakFunction(*It->Value, BreakFunc);
-			It.RemoveCurrent();
+			UE::Python::SetBreakFunction(*Pair.Value, BreakFunc);
+			return true; // Remove
 		}
-	}
+
+		return false;
+	});
 
 	NotifyModulesDirtied(DirtyModules);
 
@@ -1011,15 +1126,15 @@ void FPyWrapperTypeRegistry::OrphanWrappedTypesForModule(const FName ModuleName)
 
 void FPyWrapperTypeRegistry::GenerateWrappedTypesForReferences(const FGeneratedWrappedTypeReferences& InGeneratedWrappedTypeReferences, TSet<FName>& OutDirtyModules)
 {
-	TRACE_CPUPROFILER_EVENT_SCOPE(FPyWrapperTypeRegistry::GenerateWrappedTypesForReferences)
-
-	static const EPyTypeGenerationFlags ReferenceGenerationFlags = EPyTypeGenerationFlags::ForceShouldExport | EPyTypeGenerationFlags::IncludeBlueprintGeneratedTypes;
-
 	if (!InGeneratedWrappedTypeReferences.HasReferences())
 	{
 		return;
 	}
-	
+
+	TRACE_CPUPROFILER_EVENT_SCOPE(FPyWrapperTypeRegistry::GenerateWrappedTypesForReferences);
+
+	static const EPyTypeGenerationFlags ReferenceGenerationFlags = EPyTypeGenerationFlags::ForceShouldExport | EPyTypeGenerationFlags::IncludeBlueprintGeneratedTypes;
+
 	const int32 TotalNumReferences = InGeneratedWrappedTypeReferences.ClassReferences.Num() + InGeneratedWrappedTypeReferences.StructReferences.Num() + InGeneratedWrappedTypeReferences.EnumReferences.Num() + InGeneratedWrappedTypeReferences.DelegateReferences.Num();
 	FScopedSlowTask SlowTask(TotalNumReferences, LOCTEXT("GeneratingPythonTypes", "Generating Python Types..."));
 	if (TotalNumReferences > 1000)
@@ -1212,9 +1327,40 @@ PyTypeObject* FPyWrapperTypeRegistry::GenerateWrappedClassType(const UClass* InC
 
 	INC_DWORD_STAT(STAT_GenerateWrappedClassObjCount);
 
-	check(!GeneratedWrappedTypes.Contains(TypeRegistryName));
-	TSharedRef<PyGenUtil::FGeneratedWrappedClassType> GeneratedWrappedType = MakeShared<PyGenUtil::FGeneratedWrappedClassType>();
-	GeneratedWrappedTypes.Add(TypeRegistryName, GeneratedWrappedType);
+	TSharedPtr<PyGenUtil::FGeneratedWrappedClassType> GeneratedWrappedType = 
+		StaticCastSharedPtr<PyGenUtil::FGeneratedWrappedClassType, PyGenUtil::FGeneratedWrappedType>(
+			GeneratedWrappedTypes.FindOrProduce(
+				TypeRegistryName,
+				[]()
+				{
+					return MakeShared<PyGenUtil::FGeneratedWrappedClassType>();
+				}
+			)
+		);
+
+	if (GeneratedWrappedType->bIsReadyForUse)
+	{
+		return &GeneratedWrappedType->PyType;
+	}
+
+	TRACE_CPUPROFILER_EVENT_SCOPE(FPyWrapperTypeRegistry::GenerateWrappedClassType);
+
+	// To avoid deadlock caused by wrapped type recursion, some block of code will be run
+	// once we exit the lock of this wrapped type.
+	TArray<TFunction<void()>> CallOnceUnlocked;
+	ON_SCOPE_EXIT
+	{
+		for (TFunction<void()>& Function : CallOnceUnlocked)
+		{
+			Function();
+		}
+	};
+
+	UE::TUniqueLock GeneratedWrappedTypeScopeLock(GeneratedWrappedType->Lock);
+	if (GeneratedWrappedType->bIsReadyForUse)
+	{
+		return &GeneratedWrappedType->PyType;
+	}
 
 	TMap<FName, FName> PythonProperties;
 	TMap<FName, FString> PythonDeprecatedProperties;
@@ -1291,7 +1437,7 @@ PyTypeObject* FPyWrapperTypeRegistry::GenerateWrappedClassType(const UClass* InC
 		}
 	};
 
-	auto GenerateWrappedDynamicMethod = [this, &OutGeneratedWrappedTypeReferences, &OutDirtyModules](const UFunction* InFunc, const PyGenUtil::FGeneratedWrappedMethod& InTypeMethod)
+	auto GenerateWrappedDynamicMethod = [this, &OutGeneratedWrappedTypeReferences, &OutDirtyModules, &CallOnceUnlocked](const UFunction* InFunc, const PyGenUtil::FGeneratedWrappedMethod& InTypeMethod)
 	{
 		// Only static functions can be hoisted onto other types
 		if (!InFunc->HasAnyFunctionFlags(FUNC_Static))
@@ -1451,43 +1597,59 @@ PyTypeObject* FPyWrapperTypeRegistry::GenerateWrappedClassType(const UClass* InC
 		// Add the dynamic method to either the owner type
 		if (SelfParam.ParamProp->IsA<FObjectPropertyBase>())
 		{
-			// Ensure that we've generated a finalized Python type for this class since we'll be adding this function as a dynamic method on that type
-			const UClass* HostedClass = CastFieldChecked<const FObjectPropertyBase>(SelfParam.ParamProp)->PropertyClass;
-			if (!GenerateWrappedClassType(HostedClass, OutGeneratedWrappedTypeReferences, OutDirtyModules, EPyTypeGenerationFlags::ForceShouldExport))
-			{
-				return;
-			}
+			CallOnceUnlocked.Emplace(
+				[this, SelfParam, &OutGeneratedWrappedTypeReferences, &OutDirtyModules, DynamicMethodDefs = MoveTemp(DynamicMethodDefs), InFunc]() mutable
+				{
+					// Ensure that we've generated a finalized Python type for this class since we'll be adding this function as a dynamic method on that type
+					const UClass* HostedClass = CastFieldChecked<const FObjectPropertyBase>(SelfParam.ParamProp)->PropertyClass;
+					if (!GenerateWrappedClassType(HostedClass, OutGeneratedWrappedTypeReferences, OutDirtyModules, EPyTypeGenerationFlags::ForceShouldExport))
+					{
+						return;
+					}
 
-			// Find the wrapped type for the class as that's what we'll actually add the dynamic method to
-			TSharedPtr<PyGenUtil::FGeneratedWrappedClassType> HostedClassGeneratedWrappedType = StaticCastSharedPtr<PyGenUtil::FGeneratedWrappedClassType>(GeneratedWrappedTypes.FindRef(HostedClass));
-			check(HostedClassGeneratedWrappedType.IsValid());
+					{
+						// Find the wrapped type for the class as that's what we'll actually add the dynamic method to
+						TSharedPtr<PyGenUtil::FGeneratedWrappedClassType> HostedClassGeneratedWrappedType = StaticCastSharedPtr<PyGenUtil::FGeneratedWrappedClassType>(GeneratedWrappedTypes.FindRef(HostedClass));
+						check(HostedClassGeneratedWrappedType.IsValid());
 
-			// Add the dynamic methods to the class
-			for (TUniqueObj<PyGenUtil::FGeneratedWrappedDynamicMethod>& GeneratedWrappedDynamicMethodToAdd : DynamicMethodDefs)
-			{
-				HostedClassGeneratedWrappedType->FieldTracker.RegisterPythonFieldName(UTF8_TO_TCHAR(GeneratedWrappedDynamicMethodToAdd->MethodName.GetData()), InFunc);
-				HostedClassGeneratedWrappedType->AddDynamicMethod(MoveTemp(GeneratedWrappedDynamicMethodToAdd.Get()));
-			}
+						UE::TUniqueLock ScopedLock(HostedClassGeneratedWrappedType->Lock);
+						// Add the dynamic methods to the class
+						for (TUniqueObj<PyGenUtil::FGeneratedWrappedDynamicMethod>& GeneratedWrappedDynamicMethodToAdd : DynamicMethodDefs)
+						{
+							HostedClassGeneratedWrappedType->FieldTracker.RegisterPythonFieldName(UTF8_TO_TCHAR(GeneratedWrappedDynamicMethodToAdd->MethodName.GetData()), InFunc);
+							HostedClassGeneratedWrappedType->AddDynamicMethod(MoveTemp(GeneratedWrappedDynamicMethodToAdd.Get()));
+						}
+					}
+				}
+			);
 		}
 		else if (SelfParam.ParamProp->IsA<FStructProperty>())
 		{
-			// Ensure that we've generated a finalized Python type for this struct since we'll be adding this function as a dynamic method on that type
-			const UScriptStruct* HostedStruct = CastFieldChecked<const FStructProperty>(SelfParam.ParamProp)->Struct;
-			if (!GenerateWrappedStructType(HostedStruct, OutGeneratedWrappedTypeReferences, OutDirtyModules, EPyTypeGenerationFlags::ForceShouldExport))
-			{
-				return;
-			}
+			CallOnceUnlocked.Emplace(
+				[this, SelfParam, &OutGeneratedWrappedTypeReferences, &OutDirtyModules, DynamicMethodDefs = MoveTemp(DynamicMethodDefs), InFunc]() mutable
+				{
+					// Ensure that we've generated a finalized Python type for this struct since we'll be adding this function as a dynamic method on that type
+					const UScriptStruct* HostedStruct = CastFieldChecked<const FStructProperty>(SelfParam.ParamProp)->Struct;
+					if (!GenerateWrappedStructType(HostedStruct, OutGeneratedWrappedTypeReferences, OutDirtyModules, EPyTypeGenerationFlags::ForceShouldExport))
+					{
+						return;
+					}
 
-			// Find the wrapped type for the struct as that's what we'll actually add the dynamic method to
-			TSharedPtr<PyGenUtil::FGeneratedWrappedStructType> HostedStructGeneratedWrappedType = StaticCastSharedPtr<PyGenUtil::FGeneratedWrappedStructType>(GeneratedWrappedTypes.FindRef(HostedStruct));
-			check(HostedStructGeneratedWrappedType.IsValid());
+					{
+						// Find the wrapped type for the struct as that's what we'll actually add the dynamic method to
+						TSharedPtr<PyGenUtil::FGeneratedWrappedStructType> HostedStructGeneratedWrappedType = StaticCastSharedPtr<PyGenUtil::FGeneratedWrappedStructType>(GeneratedWrappedTypes.FindRef(HostedStruct));
+						check(HostedStructGeneratedWrappedType.IsValid());
 
-			// Add the dynamic methods to the struct
-			for (TUniqueObj<PyGenUtil::FGeneratedWrappedDynamicMethod>& GeneratedWrappedDynamicMethodToAdd : DynamicMethodDefs)
-			{
-				HostedStructGeneratedWrappedType->FieldTracker.RegisterPythonFieldName(UTF8_TO_TCHAR(GeneratedWrappedDynamicMethodToAdd->MethodName.GetData()), InFunc);
-				HostedStructGeneratedWrappedType->AddDynamicMethod(MoveTemp(GeneratedWrappedDynamicMethodToAdd.Get()));
-			}
+						UE::TUniqueLock ScopedLock(HostedStructGeneratedWrappedType->Lock);
+						// Add the dynamic methods to the struct
+						for (TUniqueObj<PyGenUtil::FGeneratedWrappedDynamicMethod>& GeneratedWrappedDynamicMethodToAdd : DynamicMethodDefs)
+						{
+							HostedStructGeneratedWrappedType->FieldTracker.RegisterPythonFieldName(UTF8_TO_TCHAR(GeneratedWrappedDynamicMethodToAdd->MethodName.GetData()), InFunc);
+							HostedStructGeneratedWrappedType->AddDynamicMethod(MoveTemp(GeneratedWrappedDynamicMethodToAdd.Get()));
+						}
+					}
+				}
+			);
 		}
 		else
 		{
@@ -1495,7 +1657,7 @@ PyTypeObject* FPyWrapperTypeRegistry::GenerateWrappedClassType(const UClass* InC
 		}
 	};
 
-	auto GenerateWrappedOperator = [this, &OutGeneratedWrappedTypeReferences, &OutDirtyModules](const UFunction* InFunc, const PyGenUtil::FGeneratedWrappedMethod& InTypeMethod)
+	auto GenerateWrappedOperator = [this, &OutGeneratedWrappedTypeReferences, &OutDirtyModules, &CallOnceUnlocked](const UFunction* InFunc, const PyGenUtil::FGeneratedWrappedMethod& InTypeMethod)
 	{
 		// Only static functions can be hoisted onto other types
 		if (!InFunc->HasAnyFunctionFlags(FUNC_Static))
@@ -1531,19 +1693,26 @@ PyTypeObject* FPyWrapperTypeRegistry::GenerateWrappedClassType(const UClass* InC
 				}
 			}
 
-			// Ensure that we've generated a finalized Python type for this struct since we'll be adding this function as a operator on that type
-			const UScriptStruct* HostedStruct = CastFieldChecked<const FStructProperty>(OpFunc.SelfParam.ParamProp)->Struct;
-			if (GenerateWrappedStructType(HostedStruct, OutGeneratedWrappedTypeReferences, OutDirtyModules, EPyTypeGenerationFlags::ForceShouldExport))
-			{
-				// Find the wrapped type for the struct as that's what we'll actually add the operator to (via its meta-data)
-				TSharedPtr<PyGenUtil::FGeneratedWrappedStructType> HostedStructGeneratedWrappedType = StaticCastSharedPtr<PyGenUtil::FGeneratedWrappedStructType>(GeneratedWrappedTypes.FindRef(HostedStruct));
-				check(HostedStructGeneratedWrappedType.IsValid());
-				StaticCastSharedPtr<FPyWrapperStructMetaData>(HostedStructGeneratedWrappedType->MetaData)->OpStacks[(int32)OpSignature.OpType].Funcs.Add(MoveTemp(OpFunc));
-			}
+			CallOnceUnlocked.Emplace(
+				[this, &OutGeneratedWrappedTypeReferences, &OutDirtyModules, OpFunc, OpSignature]() mutable
+				{
+					// Ensure that we've generated a finalized Python type for this struct since we'll be adding this function as a operator on that type
+					const UScriptStruct* HostedStruct = CastFieldChecked<const FStructProperty>(OpFunc.SelfParam.ParamProp)->Struct;
+					if (GenerateWrappedStructType(HostedStruct, OutGeneratedWrappedTypeReferences, OutDirtyModules, EPyTypeGenerationFlags::ForceShouldExport))
+					{
+						// Find the wrapped type for the struct as that's what we'll actually add the operator to (via its meta-data)
+						TSharedPtr<PyGenUtil::FGeneratedWrappedStructType> HostedStructGeneratedWrappedType = StaticCastSharedPtr<PyGenUtil::FGeneratedWrappedStructType>(GeneratedWrappedTypes.FindRef(HostedStruct));
+						check(HostedStructGeneratedWrappedType.IsValid());
+
+						UE::TUniqueLock ScopedLock(HostedStructGeneratedWrappedType->Lock);
+						StaticCastSharedPtr<FPyWrapperStructMetaData>(HostedStructGeneratedWrappedType->MetaData)->OpStacks[(int32)OpSignature.OpType].Funcs.Add(MoveTemp(OpFunc));
+					}
+				}
+			);
 		}
 	};
 
-	auto GenerateWrappedConstant = [this, &GeneratedWrappedType, &OutGeneratedWrappedTypeReferences, &OutDirtyModules](const UFunction* InFunc)
+	auto GenerateWrappedConstant = [this, &GeneratedWrappedType, &OutGeneratedWrappedTypeReferences, &OutDirtyModules, &CallOnceUnlocked](const UFunction* InFunc)
 	{
 		// Only static functions can be constants
 		if (!InFunc->HasAnyFunctionFlags(FUNC_Static))
@@ -1613,45 +1782,59 @@ PyTypeObject* FPyWrapperTypeRegistry::GenerateWrappedClassType(const UClass* InC
 		{
 			if (HostType->IsA<UClass>())
 			{
-				const UClass* HostClass = CastChecked<UClass>(HostType);
+				CallOnceUnlocked.Emplace(
+					[this, &OutGeneratedWrappedTypeReferences, &OutDirtyModules, HostType, ConstantDefs = MoveTemp(ConstantDefs), InFunc]() mutable
+					{
+						const UClass* HostClass = CastChecked<UClass>(HostType);
 
-				// Ensure that we've generated a finalized Python type for this class since we'll be adding this constant to that type
-				if (!GenerateWrappedClassType(HostClass, OutGeneratedWrappedTypeReferences, OutDirtyModules, EPyTypeGenerationFlags::ForceShouldExport))
-				{
-					return;
-				}
+						// Ensure that we've generated a finalized Python type for this class since we'll be adding this constant to that type
+						if (!GenerateWrappedClassType(HostClass, OutGeneratedWrappedTypeReferences, OutDirtyModules, EPyTypeGenerationFlags::ForceShouldExport))
+						{
+							return;
+						}
 
-				// Find the wrapped type for the class as that's what we'll actually add the constant to
-				TSharedPtr<PyGenUtil::FGeneratedWrappedClassType> HostedClassGeneratedWrappedType = StaticCastSharedPtr<PyGenUtil::FGeneratedWrappedClassType>(GeneratedWrappedTypes.FindRef(HostClass));
-				check(HostedClassGeneratedWrappedType.IsValid());
+						// Find the wrapped type for the class as that's what we'll actually add the constant to
+						TSharedPtr<PyGenUtil::FGeneratedWrappedClassType> HostedClassGeneratedWrappedType = StaticCastSharedPtr<PyGenUtil::FGeneratedWrappedClassType>(GeneratedWrappedTypes.FindRef(HostClass));
+						check(HostedClassGeneratedWrappedType.IsValid());
 
-				// Add the dynamic constants to the struct
-				for (TUniqueObj<PyGenUtil::FGeneratedWrappedConstant>& GeneratedWrappedConstantToAdd : ConstantDefs)
-				{
-					HostedClassGeneratedWrappedType->FieldTracker.RegisterPythonFieldName(UTF8_TO_TCHAR(GeneratedWrappedConstantToAdd->ConstantName.GetData()), InFunc);
-					HostedClassGeneratedWrappedType->AddDynamicConstant(MoveTemp(GeneratedWrappedConstantToAdd.Get()));
-				}
+						UE::TUniqueLock ScopedLock(HostedClassGeneratedWrappedType->Lock);
+
+						// Add the dynamic constants to the struct
+						for (TUniqueObj<PyGenUtil::FGeneratedWrappedConstant>& GeneratedWrappedConstantToAdd : ConstantDefs)
+						{
+							HostedClassGeneratedWrappedType->FieldTracker.RegisterPythonFieldName(UTF8_TO_TCHAR(GeneratedWrappedConstantToAdd->ConstantName.GetData()), InFunc);
+							HostedClassGeneratedWrappedType->AddDynamicConstant(MoveTemp(GeneratedWrappedConstantToAdd.Get()));
+						}
+					}
+				);
 			}
 			else if (HostType->IsA<UScriptStruct>())
 			{
-				const UScriptStruct* HostStruct = CastChecked<UScriptStruct>(HostType);
+				CallOnceUnlocked.Emplace(
+					[this, &OutGeneratedWrappedTypeReferences, &OutDirtyModules, HostType, ConstantDefs = MoveTemp(ConstantDefs), InFunc]() mutable
+					{
+						const UScriptStruct* HostStruct = CastChecked<UScriptStruct>(HostType);
 
-				// Ensure that we've generated a finalized Python type for this struct since we'll be adding this constant to that type
-				if (!GenerateWrappedStructType(HostStruct, OutGeneratedWrappedTypeReferences, OutDirtyModules, EPyTypeGenerationFlags::ForceShouldExport))
-				{
-					return;
-				}
+						// Ensure that we've generated a finalized Python type for this struct since we'll be adding this constant to that type
+						if (!GenerateWrappedStructType(HostStruct, OutGeneratedWrappedTypeReferences, OutDirtyModules, EPyTypeGenerationFlags::ForceShouldExport))
+						{
+							return;
+						}
 
-				// Find the wrapped type for the struct as that's what we'll actually add the constant to
-				TSharedPtr<PyGenUtil::FGeneratedWrappedStructType> HostedStructGeneratedWrappedType = StaticCastSharedPtr<PyGenUtil::FGeneratedWrappedStructType>(GeneratedWrappedTypes.FindRef(HostStruct));
-				check(HostedStructGeneratedWrappedType.IsValid());
+						// Find the wrapped type for the struct as that's what we'll actually add the constant to
+						TSharedPtr<PyGenUtil::FGeneratedWrappedStructType> HostedStructGeneratedWrappedType = StaticCastSharedPtr<PyGenUtil::FGeneratedWrappedStructType>(GeneratedWrappedTypes.FindRef(HostStruct));
+						check(HostedStructGeneratedWrappedType.IsValid());
 
-				// Add the dynamic constants to the struct
-				for (TUniqueObj<PyGenUtil::FGeneratedWrappedConstant>& GeneratedWrappedConstantToAdd : ConstantDefs)
-				{
-					HostedStructGeneratedWrappedType->FieldTracker.RegisterPythonFieldName(UTF8_TO_TCHAR(GeneratedWrappedConstantToAdd->ConstantName.GetData()), InFunc);
-					HostedStructGeneratedWrappedType->AddDynamicConstant(MoveTemp(GeneratedWrappedConstantToAdd.Get()));
-				}
+						UE::TUniqueLock ScopedLock(HostedStructGeneratedWrappedType->Lock);
+
+						// Add the dynamic constants to the struct
+						for (TUniqueObj<PyGenUtil::FGeneratedWrappedConstant>& GeneratedWrappedConstantToAdd : ConstantDefs)
+						{
+							HostedStructGeneratedWrappedType->FieldTracker.RegisterPythonFieldName(UTF8_TO_TCHAR(GeneratedWrappedConstantToAdd->ConstantName.GetData()), InFunc);
+							HostedStructGeneratedWrappedType->AddDynamicConstant(MoveTemp(GeneratedWrappedConstantToAdd.Get()));
+						}
+					}
+				);
 			}
 			else
 			{
@@ -1888,6 +2071,7 @@ PyTypeObject* FPyWrapperTypeRegistry::GenerateWrappedClassType(const UClass* InC
 			}
 		}
 
+		GeneratedWrappedType->bIsReadyForUse = true;
 		return &GeneratedWrappedType->PyType;
 	}
 	
@@ -1940,7 +2124,7 @@ PyTypeObject* FPyWrapperTypeRegistry::GenerateWrappedStructType(const UScriptStr
 	INC_DWORD_STAT(STAT_GenerateWrappedStructCallCount);
 
 	// Once we start generating types we can no longer register inline factories as they may affect the size of the generated Python objects
-	bCanRegisterInlineStructFactories = false;
+	bCanRegisterInlineStructFactories.store(false, std::memory_order_relaxed);
 
 	struct FFuncs
 	{
@@ -1983,11 +2167,31 @@ PyTypeObject* FPyWrapperTypeRegistry::GenerateWrappedStructType(const UScriptStr
 		SuperPyType = GenerateWrappedStructType(SuperStruct, OutGeneratedWrappedTypeReferences, OutDirtyModules, InGenerationFlags | EPyTypeGenerationFlags::ForceShouldExport);
 	}
 
-	INC_DWORD_STAT(STAT_GenerateWrappedStructObjCount);
+	TSharedPtr<PyGenUtil::FGeneratedWrappedStructType> GeneratedWrappedType =
+		StaticCastSharedPtr<PyGenUtil::FGeneratedWrappedStructType, PyGenUtil::FGeneratedWrappedType>(
+			GeneratedWrappedTypes.FindOrProduce(
+				TypeRegistryName,
+				[]()
+				{
+					return MakeShared<PyGenUtil::FGeneratedWrappedStructType>();
+				}
+			)
+		);
 
-	check(!GeneratedWrappedTypes.Contains(TypeRegistryName));
-	TSharedRef<PyGenUtil::FGeneratedWrappedStructType> GeneratedWrappedType = MakeShared<PyGenUtil::FGeneratedWrappedStructType>();
-	GeneratedWrappedTypes.Add(TypeRegistryName, GeneratedWrappedType);
+	if (GeneratedWrappedType->bIsReadyForUse)
+	{
+		return &GeneratedWrappedType->PyType;
+	}
+
+	TRACE_CPUPROFILER_EVENT_SCOPE(FPyWrapperTypeRegistry::GenerateWrappedStructType);
+
+	UE::TUniqueLock GeneratedWrappedTypeScopeLock(GeneratedWrappedType->Lock);
+	if (GeneratedWrappedType->bIsReadyForUse)
+	{
+		return &GeneratedWrappedType->PyType;
+	}
+
+	INC_DWORD_STAT(STAT_GenerateWrappedStructObjCount);
 
 	TMap<FName, FName> PythonProperties;
 	TMap<FName, FString> PythonDeprecatedProperties;
@@ -2189,6 +2393,7 @@ PyTypeObject* FPyWrapperTypeRegistry::GenerateWrappedStructType(const UScriptStr
 			}
 		}
 
+		GeneratedWrappedType->bIsReadyForUse = true;
 		return &GeneratedWrappedType->PyType;
 	}
 
@@ -2268,17 +2473,37 @@ PyTypeObject* FPyWrapperTypeRegistry::GenerateWrappedEnumType(const UEnum* InEnu
 
 	INC_DWORD_STAT(STAT_GenerateWrappedEnumObjCount);
 
-	TSharedPtr<PyGenUtil::FGeneratedWrappedEnumType> GeneratedWrappedType;
-	if (bIsNewType)
+	TSharedPtr<PyGenUtil::FGeneratedWrappedEnumType> GeneratedWrappedType =
+		StaticCastSharedPtr<PyGenUtil::FGeneratedWrappedEnumType, PyGenUtil::FGeneratedWrappedType>(
+			GeneratedWrappedTypes.FindOrProduce(
+				TypeRegistryName,
+				[]()
+				{
+					return MakeShared<PyGenUtil::FGeneratedWrappedEnumType>();
+				}
+			)
+		);
+
+	if (!bIsNewType)
 	{
-		check(!GeneratedWrappedTypes.Contains(TypeRegistryName));
-		GeneratedWrappedType = MakeShared<PyGenUtil::FGeneratedWrappedEnumType>();
-		GeneratedWrappedTypes.Add(TypeRegistryName, GeneratedWrappedType);
+		// this should never happen during the boot phase since we're multithreaded and we don't support
+		// handling new type since all threads rely on forward progress from other threads.
+		check(!UE::Python::Private::GIsInMultithreadedGeneration);
+		GeneratedWrappedType->bIsReadyForUse = false;
+		GeneratedWrappedType->Reset();
 	}
-	else
+
+	if (GeneratedWrappedType->bIsReadyForUse)
 	{
-		GeneratedWrappedType = StaticCastSharedPtr<PyGenUtil::FGeneratedWrappedEnumType>(GeneratedWrappedTypes.FindChecked(TypeRegistryName));
-		GeneratedWrappedType->Reset();
+		return &GeneratedWrappedType->PyType;
+	}
+
+	TRACE_CPUPROFILER_EVENT_SCOPE(FPyWrapperTypeRegistry::GenerateWrappedEnumType);
+
+	UE::TUniqueLock GeneratedWrappedTypeScopeLock(GeneratedWrappedType->Lock);
+	if (GeneratedWrappedType->bIsReadyForUse)
+	{
+		return &GeneratedWrappedType->PyType;
 	}
 
 	FString TypeDocString = PyGenUtil::PythonizeTooltip(PyGenUtil::ParseTooltip(PyGenUtil::GetFieldTooltip(InEnum)));
@@ -2371,6 +2596,7 @@ PyTypeObject* FPyWrapperTypeRegistry::GenerateWrappedEnumType(const UEnum* InEnu
 			}
 		}
 
+		GeneratedWrappedType->bIsReadyForUse = true;
 		return &GeneratedWrappedType->PyType;
 	}
 
@@ -2433,9 +2659,27 @@ PyTypeObject* FPyWrapperTypeRegistry::GenerateWrappedDelegateType(const UFunctio
 
 	INC_DWORD_STAT(STAT_GenerateWrappedDelegateObjCount);
 
-	check(!GeneratedWrappedTypes.Contains(TypeRegistryName));
-	TSharedRef<PyGenUtil::FGeneratedWrappedType> GeneratedWrappedType = MakeShared<PyGenUtil::FGeneratedWrappedType>();
-	GeneratedWrappedTypes.Add(TypeRegistryName, GeneratedWrappedType);
+	TSharedPtr<PyGenUtil::FGeneratedWrappedType> GeneratedWrappedType =
+		GeneratedWrappedTypes.FindOrProduce(
+			TypeRegistryName,
+			[]()
+			{
+				return MakeShared<PyGenUtil::FGeneratedWrappedType>();
+			}
+		);
+
+	if (GeneratedWrappedType->bIsReadyForUse)
+	{
+		return &GeneratedWrappedType->PyType;
+	}
+
+	TRACE_CPUPROFILER_EVENT_SCOPE(FPyWrapperTypeRegistry::GenerateWrappedDelegateType);
+
+	UE::TUniqueLock GeneratedWrappedTypeScopeLock(GeneratedWrappedType->Lock);
+	if (GeneratedWrappedType->bIsReadyForUse)
+	{
+		return &GeneratedWrappedType->PyType;
+	}
 
 	for (TFieldIterator<const FProperty> ParamIt(InDelegateSignature); ParamIt; ++ParamIt)
 	{
@@ -2455,15 +2699,15 @@ PyTypeObject* FPyWrapperTypeRegistry::GenerateWrappedDelegateType(const UFunctio
 	// the name of the 'PythonCallableForDelegateClass' unique because NewObject<> will find pre-existing object with the same name and
 	// return the same address, then the code below will basically erase the previous delegate type with the new delegate type, possibly
 	// confusing the parameter types if they aren't the same.
-	FNameBuilder PythonCallableForDelegateObjectName;
+	FNameBuilder NameBuilder;
 	if (UClass* OuterClass = Cast<UClass>(InDelegateSignature->GetOuter()))
 	{
 		FString DelegateOuterClassName = PyGenUtil::GetClassPythonName(OuterClass);
 
 		// The name of the UObject proxy wrapping the Python callable.
-		PythonCallableForDelegateObjectName += DelegateOuterClassName;
-		PythonCallableForDelegateObjectName += TEXT("_");
-		PythonCallableForDelegateObjectName += DelegateBaseTypename;
+		NameBuilder += DelegateOuterClassName;
+		NameBuilder += TEXT("_");
+		NameBuilder += DelegateBaseTypename;
 
 		// The name of the Python type. Note that we don't 'nest' in the pure Python sense. The type is flatten and looks like "unreal.OuterName_DelegateName"
 		// A correctly nested type would require us to update the outer object Python data structure to mention it has an inner type which we don't.
@@ -2475,33 +2719,39 @@ PyTypeObject* FPyWrapperTypeRegistry::GenerateWrappedDelegateType(const UFunctio
 	}
 	else
 	{
-		PythonCallableForDelegateObjectName += DelegateBaseTypename;
+		NameBuilder += DelegateBaseTypename;
 		GeneratedWrappedType->TypeName = PyGenUtil::TCHARToUTF8Buffer(*DelegateBaseTypename);
 	}
-	PythonCallableForDelegateObjectName += TEXT("__PythonCallable");
+	NameBuilder += TEXT("__PythonCallable");
 
-	// Generate the proxy class needed to wrap Python callables in Unreal delegates
-	UClass* PythonCallableForDelegateClass = nullptr;
-	{
-		PythonCallableForDelegateClass = NewObject<UClass>(GetPythonTypeContainer(), *PythonCallableForDelegateObjectName, RF_Public | RF_Transient);
-		UFunction* PythonCallableForDelegateFunc = nullptr;
+	FName PythonCallableForDelegateObjectName = *NameBuilder;
+
+	// This will be executed on the GameThread after the parallel initialization step
+	auto CreatePythonCallableForDelegate_GameThread =
+		[InDelegateSignature, PythonCallableForDelegateObjectName]()
 		{
-			FObjectDuplicationParameters FuncDuplicationParams(const_cast<UFunction*>(InDelegateSignature), PythonCallableForDelegateClass);
-			FuncDuplicationParams.DestName = UPythonCallableForDelegate::GeneratedFuncName;
-			FuncDuplicationParams.InternalFlagMask &= ~EInternalObjectFlags::Native;
-			PythonCallableForDelegateFunc = CastChecked<UFunction>(StaticDuplicateObjectEx(FuncDuplicationParams));
-		}
-		PythonCallableForDelegateFunc->FunctionFlags = (PythonCallableForDelegateFunc->FunctionFlags | FUNC_Native) & ~(FUNC_Delegate | FUNC_MulticastDelegate);
-		PythonCallableForDelegateFunc->SetNativeFunc(&UPythonCallableForDelegate::CallPythonNative);
-		PythonCallableForDelegateFunc->StaticLink(true);
-		PythonCallableForDelegateClass->AddFunctionToFunctionMap(PythonCallableForDelegateFunc, PythonCallableForDelegateFunc->GetFName());
-		PythonCallableForDelegateClass->SetSuperStruct(UPythonCallableForDelegate::StaticClass());
-		PythonCallableForDelegateClass->ClassFlags |= CLASS_HideDropDown;
-		PythonCallableForDelegateClass->Bind();
-		PythonCallableForDelegateClass->StaticLink(true);
-		PythonCallableForDelegateClass->AssembleReferenceTokenStream();
-		PythonCallableForDelegateClass->GetDefaultObject();
-	}
+			// Generate the proxy class needed to wrap Python callables in Unreal delegates
+			UClass* PythonCallableForDelegateClass = NewObject<UClass>(GetPythonTypeContainer(), PythonCallableForDelegateObjectName, RF_Public | RF_Transient);
+			UFunction* PythonCallableForDelegateFunc = nullptr;
+			{
+				FObjectDuplicationParameters FuncDuplicationParams(const_cast<UFunction*>(InDelegateSignature), PythonCallableForDelegateClass);
+				FuncDuplicationParams.DestName = UPythonCallableForDelegate::GeneratedFuncName;
+				FuncDuplicationParams.InternalFlagMask &= ~EInternalObjectFlags::Native;
+				PythonCallableForDelegateFunc = CastChecked<UFunction>(StaticDuplicateObjectEx(FuncDuplicationParams));
+			}
+			PythonCallableForDelegateFunc->FunctionFlags = (PythonCallableForDelegateFunc->FunctionFlags | FUNC_Native) & ~(FUNC_Delegate | FUNC_MulticastDelegate);
+			PythonCallableForDelegateFunc->SetNativeFunc(&UPythonCallableForDelegate::CallPythonNative);
+			PythonCallableForDelegateFunc->StaticLink(true);
+			PythonCallableForDelegateClass->AddFunctionToFunctionMap(PythonCallableForDelegateFunc, PythonCallableForDelegateFunc->GetFName());
+			PythonCallableForDelegateClass->SetSuperStruct(UPythonCallableForDelegate::StaticClass());
+			PythonCallableForDelegateClass->ClassFlags |= CLASS_HideDropDown;
+			PythonCallableForDelegateClass->Bind();
+			PythonCallableForDelegateClass->StaticLink(true);
+			PythonCallableForDelegateClass->AssembleReferenceTokenStream();
+			PythonCallableForDelegateClass->GetDefaultObject();
+
+			return PythonCallableForDelegateClass;
+		};
 
 	if (InDelegateSignature->HasAnyFunctionFlags(FUNC_MulticastDelegate))
 	{
@@ -2510,8 +2760,14 @@ PyTypeObject* FPyWrapperTypeRegistry::GenerateWrappedDelegateType(const UFunctio
 
 		TSharedRef<FPyWrapperMulticastDelegateMetaData> DelegateMetaData = MakeShared<FPyWrapperMulticastDelegateMetaData>();
 		DelegateMetaData->DelegateSignature.SetFunction(InDelegateSignature);
-		DelegateMetaData->PythonCallableForDelegateClass = PythonCallableForDelegateClass;
 		GeneratedWrappedType->MetaData = DelegateMetaData;
+
+		RegisteredPostInitForGameThread.ProduceItem(
+			[DelegateMetaData, CreatePythonCallableForDelegate_GameThread]()
+			{
+				DelegateMetaData->PythonCallableForDelegateClass = CreatePythonCallableForDelegate_GameThread();
+			}
+		);
 	}
 	else
 	{
@@ -2520,8 +2776,14 @@ PyTypeObject* FPyWrapperTypeRegistry::GenerateWrappedDelegateType(const UFunctio
 
 		TSharedRef<FPyWrapperDelegateMetaData> DelegateMetaData = MakeShared<FPyWrapperDelegateMetaData>();
 		DelegateMetaData->DelegateSignature.SetFunction(InDelegateSignature);
-		DelegateMetaData->PythonCallableForDelegateClass = PythonCallableForDelegateClass;
 		GeneratedWrappedType->MetaData = DelegateMetaData;
+
+		RegisteredPostInitForGameThread.ProduceItem(
+			[DelegateMetaData, CreatePythonCallableForDelegate_GameThread]()
+			{
+				DelegateMetaData->PythonCallableForDelegateClass = CreatePythonCallableForDelegate_GameThread();
+			}
+		);
 	}
 
 	if (GeneratedWrappedType->Finalize())
@@ -2544,6 +2806,7 @@ PyTypeObject* FPyWrapperTypeRegistry::GenerateWrappedDelegateType(const UFunctio
 		}
 		RegisterWrappedDelegateType(TypeRegistryName, &GeneratedWrappedType->PyType);
 
+		GeneratedWrappedType->bIsReadyForUse = true;
 		return &GeneratedWrappedType->PyType;
 	}
 
@@ -2667,6 +2930,8 @@ void FPyWrapperTypeRegistry::GatherWrappedTypesForPropertyReferences(const FProp
 
 void FPyWrapperTypeRegistry::GenerateStubCodeForWrappedTypes(const EPyOnlineDocsFilterFlags InDocGenFlags) const
 {
+	TRACE_CPUPROFILER_EVENT_SCOPE(FPyWrapperTypeRegistry::GenerateStubCodeForWrappedTypes)
+
 	UE_LOG(LogPython, Display, TEXT("Generating Python API stub file..."));
 
 	// Suspend the hang and hitch heartbeats, as this is a long running task.
@@ -2842,7 +3107,7 @@ void FPyWrapperTypeRegistry::GenerateStubCodeForWrappedTypes(const EPyOnlineDocs
 	// Process generated glue code
 	// Also excludes types that don't pass the filters specified in InDocGenFlags using the information about
 	// which module it came from and where that module exists on disk.
-	auto ProcessWrappedDataArray = [this, &PythonScript, InDocGenFlags](const TMap<FSoftObjectPath, PyTypeObject*>& WrappedData, const TSharedPtr<FPyOnlineDocsSection>& OnlineDocsSection)
+	auto ProcessWrappedDataArray = [this, &PythonScript, InDocGenFlags](const auto& WrappedData, const TSharedPtr<FPyOnlineDocsSection>& OnlineDocsSection)
 	{
 		if (OnlineDocsSection.IsValid())
 		{
@@ -2857,7 +3122,7 @@ void FPyWrapperTypeRegistry::GenerateStubCodeForWrappedTypes(const EPyOnlineDocs
 			ProjectTopDir = FPaths::GetCleanFilename(FPaths::ProjectDir());
 		}
 
-		for (const auto& WrappedDataPair : WrappedData)
+		WrappedData.ForEach([this, &PythonScript, InDocGenFlags, &ProjectTopDir, &OnlineDocsSection](const auto& WrappedDataPair)
 		{
 			TSharedPtr<PyGenUtil::FGeneratedWrappedType> GeneratedWrappedType = GeneratedWrappedTypes.FindRef(WrappedDataPair.Key);
 
@@ -2881,7 +3146,7 @@ void FPyWrapperTypeRegistry::GenerateStubCodeForWrappedTypes(const EPyOnlineDocs
 						// Optionally exclude Project classes
 						if (!EnumHasAnyFlags(InDocGenFlags, EPyOnlineDocsFilterFlags::IncludeProject))
 						{
-							continue;
+							return;
 						}
 					}
 					// Is Enterprise class
@@ -2890,7 +3155,7 @@ void FPyWrapperTypeRegistry::GenerateStubCodeForWrappedTypes(const EPyOnlineDocs
 						// Optionally exclude Enterprise classes
 						if (!EnumHasAnyFlags(InDocGenFlags, EPyOnlineDocsFilterFlags::IncludeEnterprise))
 						{
-							continue;
+							return;
 						}
 					}
 					// is internal class
@@ -2899,7 +3164,7 @@ void FPyWrapperTypeRegistry::GenerateStubCodeForWrappedTypes(const EPyOnlineDocs
 						// Optionally exclude internal classes
 						if (!EnumHasAnyFlags(InDocGenFlags, EPyOnlineDocsFilterFlags::IncludeInternal))
 						{
-							continue;
+							return;
 						}
 					}
 					// Everything else is considered an "Engine" class
@@ -2908,7 +3173,7 @@ void FPyWrapperTypeRegistry::GenerateStubCodeForWrappedTypes(const EPyOnlineDocs
 						// Optionally exclude engine classes
 						if (!EnumHasAnyFlags(InDocGenFlags, EPyOnlineDocsFilterFlags::IncludeEngine))
 						{
-							continue;
+							return;
 						}
 					}
 				}
@@ -2916,7 +3181,7 @@ void FPyWrapperTypeRegistry::GenerateStubCodeForWrappedTypes(const EPyOnlineDocs
 			}
 
 			GenerateStubCodeForWrappedType(WrappedDataPair.Value, GeneratedWrappedType.Get(), PythonScript, OnlineDocsSection.Get());
-		}
+		});
 	};
 
 	ProcessWrappedDataArray(PythonWrappedEnums, OnlineDocsEnumTypesSection);
@@ -3732,24 +3997,33 @@ void FPyWrapperTypeRegistry::GenerateStubCodeForWrappedType(PyTypeObject* PyType
 
 void FPyWrapperTypeRegistry::RegisterPythonTypeName(const FString& InPythonTypeName, const FSoftObjectPath& InUnrealTypeName)
 {
-	FSoftObjectPath& ExistingUnrealTypeName = PythonWrappedTypeNameToUnrealTypeName.FindOrAdd(InPythonTypeName);
-	if (ExistingUnrealTypeName.IsNull())
-	{
-		ExistingUnrealTypeName = InUnrealTypeName;
-	}
-	else
-	{
-		REPORT_PYTHON_GENERATION_ISSUE(Warning, TEXT("'%s' and '%s' have the same name (%s) when exposed to Python. Rename one of them using 'ScriptName' meta-data."), *ExistingUnrealTypeName.ToString(), *InUnrealTypeName.ToString(), *InPythonTypeName);
-	}
+	bool bAlreadyExists = true;
+	PythonWrappedTypeNameToUnrealTypeName.FindOrProduceAndApply(
+		InPythonTypeName,
+		[&InUnrealTypeName, &bAlreadyExists]()
+		{
+			// If the value is being added, then we don't have to perform the additional test in Apply
+			bAlreadyExists = false;
+			return InUnrealTypeName;
+		},
+		[&InPythonTypeName, &InUnrealTypeName, &bAlreadyExists](const FSoftObjectPath& Value)
+		{
+			if (bAlreadyExists && Value != InUnrealTypeName)
+			{
+				REPORT_PYTHON_GENERATION_ISSUE(Warning, TEXT("'%s' and '%s' have the same name (%s) when exposed to Python. Rename one of them using 'ScriptName' meta-data."), *Value.ToString(), *InUnrealTypeName.ToString(), *InPythonTypeName);
+			}
+		}
+	);
 }
 
 void FPyWrapperTypeRegistry::UnregisterPythonTypeName(const FString& InPythonTypeName, const FSoftObjectPath& InUnrealTypeName)
 {
-	const FSoftObjectPath* ExistingUnrealTypeNamePtr = PythonWrappedTypeNameToUnrealTypeName.Find(InPythonTypeName);
-	if (ExistingUnrealTypeNamePtr && *ExistingUnrealTypeNamePtr == InUnrealTypeName)
-	{
-		PythonWrappedTypeNameToUnrealTypeName.Remove(InPythonTypeName);
-	}
+	PythonWrappedTypeNameToUnrealTypeName.RemoveIf(InPythonTypeName,
+		[&InUnrealTypeName](const FSoftObjectPath& ExistingUnrealTypeName)
+		{
+			return ExistingUnrealTypeName == InUnrealTypeName;
+		}
+	);
 }
 
 #endif	// WITH_PYTHON
diff --git a/Engine/Plugins/Experimental/PythonScriptPlugin/Source/PythonScriptPlugin/Private/PyWrapperTypeRegistry.h b/Engine/Plugins/Experimental/PythonScriptPlugin/Source/PythonScriptPlugin/Private/PyWrapperTypeRegistry.h
--- a/Engine/Plugins/Experimental/PythonScriptPlugin/Source/PythonScriptPlugin/Private/PyWrapperTypeRegistry.h
+++ b/Engine/Plugins/Experimental/PythonScriptPlugin/Source/PythonScriptPlugin/Private/PyWrapperTypeRegistry.h
@@ -12,8 +12,12 @@
 #include "Misc/EnumClassFlags.h"
 #include "UObject/GCObject.h"
 #include "UObject/WeakObjectPtr.h"
+#include "Containers/ConsumeAllMpmcQueue.h"
+#include "Containers/StripedMap.h"
 #include "Templates/Function.h"
 
+#include <atomic>
+
 #if WITH_PYTHON
 
 struct FPyWrapperObject;
@@ -390,7 +394,7 @@ public:
 	const IPyWrapperInlineStructFactory* GetInlineStructFactory(const FTopLevelAssetPath& StructName) const;
 
 	/** Generate a wrapped type for all currently available Unreal types */
-	void GenerateWrappedTypes();
+	void GenerateWrappedTypes(bool bAllowMultithreadedGeneration = false);
 
 	/** Generate a wrapped type for all available Unreal types in the given Unreal module */
 	void GenerateWrappedTypesForModule(const FName ModuleName);
@@ -486,6 +490,9 @@ public:
 	void GenerateStubCodeForWrappedTypes(const EPyOnlineDocsFilterFlags InDocGenFlags = EPyOnlineDocsFilterFlags::IncludeNone) const;
 
 private:
+	/** Call the game-thread only logic */
+	void ExecutePostInitOnGameThread();
+
 	/** Gather any types referenced by the given property are still need to be wrapped for use in Python */
 	void GatherWrappedTypesForPropertyReferences(const FProperty* InProp, FGeneratedWrappedTypeReferences& OutGeneratedWrappedTypeReferences) const;
 
@@ -501,38 +508,41 @@ private:
 	FPyWrapperTypeRegistry();
 
 	/** True if we can register inline struct factories, False if we've already started generating types */
-	bool bCanRegisterInlineStructFactories;
+	std::atomic<bool> bCanRegisterInlineStructFactories;
 
 	/** Map from the Unreal struct name to the factory data for an inline struct (ie, a struct known at compile time that will allocate its instance data inlined within the Python object) */
 	TMap<FTopLevelAssetPath, TSharedPtr<const IPyWrapperInlineStructFactory>> InlineStructFactories;
 
+	// Good balance for high core count CPUs.
+	static constexpr int32 StripeCount = 32;
+
 	/** Map from the Python wrapped type name to the Unreal type name it was generated from (for conflict detection) */
-	typedef TMap<FString, FSoftObjectPath, FDefaultSetAllocator, PyGenUtil::FCaseSensitiveStringMapFuncs<FSoftObjectPath>> FCaseSensitiveStringToNameMap;
+	typedef TStripedMap<StripeCount, FString, FSoftObjectPath, FDefaultSetAllocator, PyGenUtil::FCaseSensitiveStringMapFuncs<FSoftObjectPath>> FCaseSensitiveStringToNameMap;
 	FCaseSensitiveStringToNameMap PythonWrappedTypeNameToUnrealTypeName;
 
 	/** Map from the Unreal class name to the Python type */
-	TMap<FSoftObjectPath, PyTypeObject*> PythonWrappedClasses;
+	TStripedMap<StripeCount, FSoftObjectPath, PyTypeObject*> PythonWrappedClasses;
 
 	/** Map from the Unreal struct name to the Python type */
-	TMap<FSoftObjectPath, PyTypeObject*> PythonWrappedStructs;
+	TStripedMap<StripeCount, FSoftObjectPath, PyTypeObject*> PythonWrappedStructs;
 
 	/** Map from the Unreal enum name to the Python type */
-	TMap<FSoftObjectPath, PyTypeObject*> PythonWrappedEnums;
+	TStripedMap<StripeCount, FSoftObjectPath, PyTypeObject*> PythonWrappedEnums;
 
 	/** Map from the Unreal delegate signature name to the Python type */
-	TMap<FSoftObjectPath, PyTypeObject*> PythonWrappedDelegates;
+	TStripedMap<StripeCount, FSoftObjectPath, PyTypeObject*> PythonWrappedDelegates;
 
 	/** Map from the Unreal type name to the generated Python type data */
-	TMap<FSoftObjectPath, TSharedPtr<PyGenUtil::FGeneratedWrappedType>> GeneratedWrappedTypes;
+	TStripedMap<StripeCount, FSoftObjectPath, TSharedPtr<PyGenUtil::FGeneratedWrappedType>> GeneratedWrappedTypes;
 
 	/** Map from the Unreal module name to its generated type names (names are the Unreal names) */
-	TMultiMap<FName, FSoftObjectPath> GeneratedWrappedTypesForModule;
+	TStripedMultiMap<StripeCount, FName, FSoftObjectPath> GeneratedWrappedTypesForModule;
 
 	/** Map make func name to the meta data of the struct to make. Used when meta 'HasNativeMake' references a function not loaded yet. */
-	TMap<FString, TSharedPtr<FPyWrapperStructMetaData>> UnresolvedMakeFuncs;
+	TStripedMap<StripeCount, FString, TSharedPtr<FPyWrapperStructMetaData>> UnresolvedMakeFuncs;
 
 	/** Map break func name to the meta data of the struct to break. Used when meta 'HasNativeBreak' references a function not loaded yet. */
-	TMap<FString, TSharedPtr<FPyWrapperStructMetaData>> UnresolvedBreakFuncs;
+	TStripedMap<StripeCount, FString, TSharedPtr<FPyWrapperStructMetaData>> UnresolvedBreakFuncs;
 
 	/** Array of generated Python type data that has been orphaned (due to its owner module being unloaded/reloaded) */
 	TArray<TSharedPtr<PyGenUtil::FGeneratedWrappedType>> OrphanedWrappedTypes;
@@ -540,6 +550,9 @@ private:
 	/** Array of information about native Python modules */
 	TArray<PyGenUtil::FNativePythonModule> NativePythonModules;
 
+	/** Functions to be called on the game thread after or during a multi-threaded generation pass. */
+	UE::TConsumeAllMpmcQueue<TFunction<void ()>> RegisteredPostInitForGameThread;
+
 	/** Callback for when a Python module is dirtied */
 	FOnModuleDirtied OnModuleDirtiedDelegate;
 };
diff --git a/Engine/Source/Runtime/Core/Private/Internationalization/TextCache.cpp b/Engine/Source/Runtime/Core/Private/Internationalization/TextCache.cpp
--- a/Engine/Source/Runtime/Core/Private/Internationalization/TextCache.cpp
+++ b/Engine/Source/Runtime/Core/Private/Internationalization/TextCache.cpp
@@ -31,23 +31,23 @@ FText FTextCache::FindOrCache(FStringView InTextLiteral, const FTextId& InTextId
 
 	// First try and find a cached instance
 	{
-		FText* ReturnFoundText = nullptr;
-
-		UE::TScopeLock Lock(CachedTextCS);
-
-		FText* FoundText = CachedText.Find(InTextId);
-		if (FoundText)
-		{
-			const FString* FoundTextLiteral = FTextInspector::GetSourceString(*FoundText);
-			if (FoundTextLiteral && InTextLiteral.Equals(*FoundTextLiteral, ESearchCase::CaseSensitive))
+		FText Text;
+		bool bTextFound = false;
+		CachedText.FindAndApply(InTextId,
+			[&InTextLiteral, &Text, &bTextFound](const FText& FoundText)
 			{
-				ReturnFoundText = FoundText;
+				const FString* FoundTextLiteral = FTextInspector::GetSourceString(FoundText);
+				if (FoundTextLiteral && InTextLiteral.Equals(*FoundTextLiteral, ESearchCase::CaseSensitive))
+				{
+					Text = FoundText;
+					bTextFound = true;
+				}
 			}
-		}
+		);
 
-		if (ReturnFoundText)
+		if (bTextFound)
 		{
-			return *ReturnFoundText;
+			return Text;
 		}
 	}
 
@@ -55,9 +55,7 @@ FText FTextCache::FindOrCache(FStringView InTextLiteral, const FTextId& InTextId
 	FText NewText = FText(FString(InTextLiteral), InTextId.GetNamespace(), InTextId.GetKey(), ETextFlag::Immutable);
 
 	// ... and add it to the cache
-	UE::TScopeLock Lock(CachedTextCS);
-
-	CachedText.Emplace(InTextId, NewText);
+	CachedText.Add(InTextId, NewText);
 
 	return NewText;
 }
@@ -68,23 +66,23 @@ FText FTextCache::FindOrCache(FString&& InTextLiteral, const FTextId& InTextId)
 
 	// First try and find a cached instance
 	{
-		FText* ReturnFoundText = nullptr;
-
-		UE::TScopeLock Lock(CachedTextCS);
-
-		FText* FoundText = CachedText.Find(InTextId);
-		if (FoundText)
-		{
-			const FString* FoundTextLiteral = FTextInspector::GetSourceString(*FoundText);
-			if (FoundTextLiteral && InTextLiteral.Equals(*FoundTextLiteral, ESearchCase::CaseSensitive))
+		FText Text;
+		bool bTextFound = false;
+		CachedText.FindAndApply(InTextId,
+			[&InTextLiteral, &Text, &bTextFound](const FText& FoundText)
 			{
-				ReturnFoundText = FoundText;
+				const FString* FoundTextLiteral = FTextInspector::GetSourceString(FoundText);
+				if (FoundTextLiteral && InTextLiteral.Equals(*FoundTextLiteral, ESearchCase::CaseSensitive))
+				{
+					Text = FoundText;
+					bTextFound = true;
+				}
 			}
-		}
+		);
 
-		if (ReturnFoundText)
+		if (bTextFound)
 		{
-			return *ReturnFoundText;
+			return Text;
 		}
 	}
 
@@ -92,9 +90,7 @@ FText FTextCache::FindOrCache(FString&& InTextLiteral, const FTextId& InTextId)
 	FText NewText = FText(MoveTemp(InTextLiteral), InTextId.GetNamespace(), InTextId.GetKey(), ETextFlag::Immutable);
 
 	// ... and add it to the cache
-	UE::TScopeLock Lock(CachedTextCS);
-
-	CachedText.Emplace(InTextId, NewText);
+	CachedText.Add(InTextId, NewText);
 
 	return NewText;
 }
@@ -106,7 +102,6 @@ void FTextCache::RemoveCache(const FTextId& InTextId)
 
 void FTextCache::RemoveCache(TArrayView<const FTextId> InTextIds)
 {
-	UE::TScopeLock Lock(CachedTextCS);
 	for (const FTextId& TextId : InTextIds)
 	{
 		CachedText.Remove(TextId);
@@ -115,7 +110,6 @@ void FTextCache::RemoveCache(TArrayView<const FTextId> InTextIds)
 
 void FTextCache::RemoveCache(const TSet<FTextId>& InTextIds)
 {
-	UE::TScopeLock Lock(CachedTextCS);
 	for (const FTextId& TextId : InTextIds)
 	{
 		CachedText.Remove(TextId);
diff --git a/Engine/Source/Runtime/Core/Private/Internationalization/TextCache.h b/Engine/Source/Runtime/Core/Private/Internationalization/TextCache.h
--- a/Engine/Source/Runtime/Core/Private/Internationalization/TextCache.h
+++ b/Engine/Source/Runtime/Core/Private/Internationalization/TextCache.h
@@ -3,6 +3,7 @@
 #pragma once
 
 #include "Containers/Map.h"
+#include "Containers/StripedMap.h"
 #include "HAL/Platform.h"
 #include "Internationalization/Text.h"
 #include "Misc/TransactionallySafeCriticalSection.h"
@@ -40,6 +41,5 @@ public:
 	void RemoveCache(const TSet<FTextId>& InTextIds);
 
 private:
-	TMap<FTextId, FText> CachedText;
-	FTransactionallySafeCriticalSection CachedTextCS;
+	TStripedMap<32, FTextId, FText, FDefaultSetAllocator, TDefaultMapHashableKeyFuncs<FTextId, FText, false>, FTransactionallySafeStripedMapLockingPolicy> CachedText;
 };
