[Python]
- Parallelize FPipInstall::WriteSitePackagePthFile
- 177ms to 16ms

diff --git a/Engine/Plugins/Experimental/PythonScriptPlugin/Source/PythonScriptPlugin/Private/PipInstall.cpp b/Engine/Plugins/Experimental/PythonScriptPlugin/Source/PythonScriptPlugin/Private/PipInstall.cpp
--- a/Engine/Plugins/Experimental/PythonScriptPlugin/Source/PythonScriptPlugin/Private/PipInstall.cpp
+++ b/Engine/Plugins/Experimental/PythonScriptPlugin/Source/PythonScriptPlugin/Private/PipInstall.cpp
@@ -7,6 +7,8 @@
 #include "PythonScriptPluginSettings.h"
 #include "Dom/JsonObject.h"
 #include "HAL/PlatformFileManager.h"
+#include "Async/Mutex.h"
+#include "Async/ParallelFor.h"
 #include "Misc/CommandLine.h"
 #include "Misc/FeedbackContext.h"
 #include "Misc/FileHelper.h"
@@ -426,51 +428,70 @@ FPipInstall::FPipInstall()
 
 void FPipInstall::WriteSitePackagePthFile() const
 {
+	TRACE_CPUPROFILER_EVENT_SCOPE(FPipInstall::WriteSitePackagePthFile);
+
 	// Write all paths from script-plugin
 	// TODO: Should we directly use PyUtil::GetSystemPaths instead?
 	// TArray<FString> PluginSitePackagePaths = PyUtil::GetSystemPaths();
 
 	// List of enabled plugins' site-packages folders
 	TArray<FString> PluginSitePackagePaths;
-	for ( const TSharedRef<IPlugin>& Plugin : IPluginManager::Get().GetEnabledPlugins() )
-	{
-		const FString PythonContentPath = FPaths::ConvertRelativePathToFull(Plugin->GetContentDir() / TEXT("Python"));
-		const FString PluginPlatformSitePackagesPath = PythonContentPath / TEXT("Lib") / FPlatformMisc::GetUBTPlatform() / TEXT("site-packages");
-		const FString PluginGeneralSitePackagesPath = PythonContentPath / TEXT("Lib") / TEXT("site-packages");
+	UE::FMutex Mutex;
 
-		// Write platform/general site-packages paths per-plugin to .pth file to account for packaged python dependencies during pip install
-		if (FPaths::DirectoryExists(PluginPlatformSitePackagesPath))
+	TArray<TSharedRef<IPlugin>> Plugins = IPluginManager::Get().GetEnabledPlugins();
+	ParallelFor(Plugins.Num(),
+		[&PluginSitePackagePaths, &Plugins, &Mutex](int32 Index)
 		{
-			PluginSitePackagePaths.Add(PluginPlatformSitePackagesPath);
-		}
+			TSharedRef<IPlugin> Plugin = Plugins[Index];
+			const FString PythonContentPath = FPaths::ConvertRelativePathToFull(Plugin->GetContentDir() / TEXT("Python"));
+			FString PluginPlatformSitePackagesPath = PythonContentPath / TEXT("Lib") / FPlatformMisc::GetUBTPlatform() / TEXT("site-packages");
+			FString PluginGeneralSitePackagesPath = PythonContentPath / TEXT("Lib") / TEXT("site-packages");
 
-		if (FPaths::DirectoryExists(PluginGeneralSitePackagesPath))
-		{
-			PluginSitePackagePaths.Add(PluginGeneralSitePackagesPath);
+			// Write platform/general site-packages paths per-plugin to .pth file to account for packaged python dependencies during pip install
+			const bool bPlatformSitePackageExists = FPaths::DirectoryExists(PluginPlatformSitePackagesPath);
+			const bool bGeneralSitePackageExists = FPaths::DirectoryExists(PluginGeneralSitePackagesPath);
+
+			if (bPlatformSitePackageExists || bGeneralSitePackageExists)
+			{
+				UE::TUniqueLock ScopeLock(Mutex);
+				if (bPlatformSitePackageExists)
+				{
+					PluginSitePackagePaths.Emplace(MoveTemp(PluginPlatformSitePackagesPath));
+				}
+				if (bGeneralSitePackageExists)
+				{
+					PluginSitePackagePaths.Emplace(MoveTemp(PluginGeneralSitePackagesPath));
+				}
+			}
 		}
-	}
+	);
 
 	// Additional paths
 	for (const FDirectoryPath& AdditionalPath : GetDefault<UPythonScriptPluginSettings>()->AdditionalPaths)
 	{
-		const FString AddPath = FPaths::ConvertRelativePathToFull(AdditionalPath.Path);
+		FString AddPath = FPaths::ConvertRelativePathToFull(AdditionalPath.Path);
 		if (FPaths::DirectoryExists(AddPath))
 		{
-			PluginSitePackagePaths.Add(AddPath);
+			PluginSitePackagePaths.Emplace(MoveTemp(AddPath));
 		}
 	}
 
 	// UE_PYTHONPATH
-	TArray<FString> SystemEnvPaths;
-	FPlatformMisc::GetEnvironmentVariable(TEXT("UE_PYTHONPATH")).ParseIntoArray(SystemEnvPaths, FPlatformMisc::GetPathVarDelimiter());
-	for (const FString& SystemEnvPath : SystemEnvPaths)
 	{
-		if (FPaths::DirectoryExists(SystemEnvPath))
+		TArray<FString> SystemEnvPaths;
+		FPlatformMisc::GetEnvironmentVariable(TEXT("UE_PYTHONPATH")).ParseIntoArray(SystemEnvPaths, FPlatformMisc::GetPathVarDelimiter());
+		for (FString& SystemEnvPath : SystemEnvPaths)
 		{
-			PluginSitePackagePaths.Add(SystemEnvPath);
+			if (FPaths::DirectoryExists(SystemEnvPath))
+			{
+				PluginSitePackagePaths.Emplace(MoveTemp(SystemEnvPath));
+			}
 		}
 	}
 
+	// Make sure the order in the file is deterministic
+	PluginSitePackagePaths.Sort();
+
 	// Create .pth file in PipInstall/Lib/site-packages to account for plugins with packaged dependencies
 	const FString PyPluginsSitePackageFile = FPaths::ConvertRelativePathToFull(GetPipSitePackagesPath() / PluginsSitePackageFilename);
 	FFileHelper::SaveStringArrayToFile(PluginSitePackagePaths, *PyPluginsSitePackageFile);
