[Python]
- Move PipInstall out-of-process calls to another thread during type wrapper generation
- Reduce the sleeping length from 100ms to 10ms while we wait on out-of-process calls to finish
- 838ms to 0ms on the game-thread

diff --git a/Engine/Plugins/Experimental/PythonScriptPlugin/Source/PythonScriptPlugin/Private/PipRunnable.cpp b/Engine/Plugins/Experimental/PythonScriptPlugin/Source/PythonScriptPlugin/Private/PipRunnable.cpp
--- a/Engine/Plugins/Experimental/PythonScriptPlugin/Source/PythonScriptPlugin/Private/PipRunnable.cpp
+++ b/Engine/Plugins/Experimental/PythonScriptPlugin/Source/PythonScriptPlugin/Private/PipRunnable.cpp
@@ -192,7 +192,7 @@ bool FLoggedSubprocessSync::Run(int32& OutExitCode, const FString& URI, const FS
 
 	while (Subproc.TickSubprocess())
 	{
-		FPlatformProcess::Sleep(0.1f);
+		FPlatformProcess::Sleep(0.01f);
 	}
 
 	OutExitCode = Subproc.GetExitCode();
diff --git a/Engine/Plugins/Experimental/PythonScriptPlugin/Source/PythonScriptPlugin/Private/PythonScriptPlugin.cpp b/Engine/Plugins/Experimental/PythonScriptPlugin/Source/PythonScriptPlugin/Private/PythonScriptPlugin.cpp
--- a/Engine/Plugins/Experimental/PythonScriptPlugin/Source/PythonScriptPlugin/Private/PythonScriptPlugin.cpp
+++ b/Engine/Plugins/Experimental/PythonScriptPlugin/Source/PythonScriptPlugin/Private/PythonScriptPlugin.cpp
@@ -50,6 +50,7 @@
 #include "Framework/MultiBox/MultiBoxBuilder.h"
 #include "ToolMenus.h"
 #include "Misc/ConfigCacheIni.h"
+#include "Tasks/Task.h"
 #include "Misc/MessageDialog.h"
 #include "AssetViewUtils.h"
 #include "IContentBrowserDataModule.h"
@@ -1248,14 +1249,28 @@ void FPythonScriptPlugin::InitializePython()
 		FPyWrapperTypeRegistry::Get().OnModuleDirtied().AddRaw(this, &FPythonScriptPlugin::OnModuleDirtied);
 		FModuleManager::Get().OnModulesChanged().AddRaw(this, &FPythonScriptPlugin::OnModulesChanged);
 
+#if WITH_EDITOR
+		UE::Tasks::FTask PipTask = 
+			UE::Tasks::Launch(
+				TEXT("PipInstallAsync"),
+				[this]()
+				{
+					// Init PipInstall task
+					InitPipInstaller();
+					RunPipInstaller();
+				}
+			);
+#endif // WITH_EDITOR
+
+		// We only allow that during boot
+		const bool bAllowMultithreadedGeneration = true;
+
 		// Initialize the wrapped types
-		FPyWrapperTypeRegistry::Get().GenerateWrappedTypes();
+		FPyWrapperTypeRegistry::Get().GenerateWrappedTypes(bAllowMultithreadedGeneration);
 
 #if WITH_EDITOR
-		// Init PipInstall task
-		InitPipInstaller();
-		RunPipInstaller();
-#endif // WITH_EDITOR
+		PipTask.Wait();
+#endif
 
 		// Add Pip UBT install path to site-packages if it exists
 		const FString PipSitePackagePath = FPipInstall::Get().GetPipSitePackagesPath();
