[Python]
- Make CIS happy about unreachable code due to break placement

diff --git a/Engine/Plugins/Experimental/PythonScriptPlugin/Source/PythonScriptPlugin/Private/PyWrapperTypeRegistry.cpp b/Engine/Plugins/Experimental/PythonScriptPlugin/Source/PythonScriptPlugin/Private/PyWrapperTypeRegistry.cpp
--- a/Engine/Plugins/Experimental/PythonScriptPlugin/Source/PythonScriptPlugin/Private/PyWrapperTypeRegistry.cpp
+++ b/Engine/Plugins/Experimental/PythonScriptPlugin/Source/PythonScriptPlugin/Private/PyWrapperTypeRegistry.cpp
@@ -1458,14 +1458,14 @@ PyTypeObject* FPyWrapperTypeRegistry::GenerateWrappedClassType(const UClass* InC
 		{
 			// Hint that UPRAM(ref) may be missing on first parameter
 			const FProperty* PropertyPossiblyMissingMacro = nullptr;
-			for (TFieldIterator<const FProperty> ParamIt(InFunc); ParamIt; ++ParamIt)
+			for (TFieldIterator<const FProperty> ParamIt(InFunc); ParamIt; )
 			{
 				const FProperty* Param = *ParamIt;
 				if (PyUtil::IsOutputParameter(Param) && !PyUtil::IsInputParameter(Param) && !Param->HasAnyPropertyFlags(CPF_ReturnParm) && IsStructOrObjectProperty(Param))
 				{
 					PropertyPossiblyMissingMacro = Param;
 				}
-				break;
+				break; // This is properly placed, we want to only check the first property.
 			}
 
 			REPORT_PYTHON_GENERATION_ISSUE(Error, TEXT("Function '%s.%s' is marked as 'ScriptMethod' but doesn't contain a valid struct or object as its first argument.%s"), *InFunc->GetOwnerClass()->GetPathName(), *InFunc->GetName(), PropertyPossiblyMissingMacro ? TEXT(" UPARAM(ref) may be missing on the first argument.") : TEXT(""));
