Add redirector for the PythonGeneratedClass CDO (from the old short name to the new unique name)

diff --git a/Engine/Plugins/Experimental/PythonScriptPlugin/Source/PythonScriptPlugin/Private/PyWrapperObject.cpp b/Engine/Plugins/Experimental/PythonScriptPlugin/Source/PythonScriptPlugin/Private/PyWrapperObject.cpp
--- a/Engine/Plugins/Experimental/PythonScriptPlugin/Source/PythonScriptPlugin/Private/PyWrapperObject.cpp
+++ b/Engine/Plugins/Experimental/PythonScriptPlugin/Source/PythonScriptPlugin/Private/PyWrapperObject.cpp
@@ -1420,12 +1420,28 @@ public:
 			PrepareOldClassForReinstancing();
 		}
 
+		const FString CleanPyTypeName = PyUtil::GetCleanTypename(PyType);
+
+		FNameBuilder CleanCDOName;
+		CleanCDOName.Append(DEFAULT_OBJECT_PREFIX);
+		CleanCDOName.Append(CleanPyTypeName);
+
 		Py_BEGIN_ALLOW_THREADS
-		NewClass->ClassRedirector = CreatePythonTypeLegacyRedirector(PyUtil::GetCleanTypename(PyType), FTopLevelAssetPath(NewClass->GetOuter()->GetFName(), *ClassName));
+		NewClass->ClassRedirector = CreatePythonTypeLegacyRedirector(CleanPyTypeName, FTopLevelAssetPath(NewClass->GetOuter()->GetFName(), *ClassName));
 		if (NewClass->ClassRedirector)
 		{
 			NewClass->ClassRedirector->DestinationObject = NewClass;
 		}
+		else
+		{
+			// If the class doesn't have a redirector, then we also need to remove any existing redirector for the CDO (before the real CDO is created below)
+			if (UObjectRedirector* ExistingCDORedirector = FindObject<UObjectRedirector>(GetPythonTypeContainer(), *CleanCDOName))
+			{
+				ExistingCDORedirector->DestinationObject = nullptr;
+				ExistingCDORedirector->ClearFlags(RF_Public | RF_Standalone);
+				ExistingCDORedirector->Rename(*FString::Printf(TEXT("%s_DELETED"), *CleanCDOName), nullptr, REN_DontCreateRedirectors);
+			}
+		}
 		NewClass->Rename(*ClassName, nullptr, REN_DontCreateRedirectors);
 
 		// Finalize the class
@@ -1442,9 +1458,17 @@ public:
 		NewClass->PyType = FPyTypeObjectPtr::NewReference(PyType);
 		FPyWrapperTypeRegistry::Get().RegisterWrappedClassType(NewClass, PyType, false);
 
-		// Ensure the CDO exists
+		// Ensure the CDO exists and add a redirector for it if needed
 		Py_BEGIN_ALLOW_THREADS
-		NewClass->GetDefaultObject();
+		if (UObject* NewClassCDO = NewClass->GetDefaultObject())
+		{
+			if (NewClass->ClassRedirector)
+			{
+				// If the class has a redirector, then we also need to redirect the CDO
+				NewClass->CDORedirector = NewObject<UObjectRedirector>(GetPythonTypeContainer(), FName(CleanCDOName), RF_Public | RF_Transient);
+				NewClass->CDORedirector->DestinationObject = NewClassCDO;
+			}
+		}
 		Py_END_ALLOW_THREADS
 
 		// Re-instance the old class and re-parent any derived classes to this new type
@@ -1977,6 +2001,7 @@ void UPythonGeneratedClass::AddReferencedObjects(UObject* InThis, FReferenceColl
 
 	UPythonGeneratedClass* This = CastChecked<UPythonGeneratedClass>(InThis);
 	Collector.AddReferencedObject(This->ClassRedirector, This);
+	Collector.AddReferencedObject(This->CDORedirector, This);
 }
 
 void UPythonGeneratedClass::PostRename(UObject* OldOuter, const FName OldName)
diff --git a/Engine/Plugins/Experimental/PythonScriptPlugin/Source/PythonScriptPlugin/Private/PyWrapperObject.h b/Engine/Plugins/Experimental/PythonScriptPlugin/Source/PythonScriptPlugin/Private/PyWrapperObject.h
--- a/Engine/Plugins/Experimental/PythonScriptPlugin/Source/PythonScriptPlugin/Private/PyWrapperObject.h
+++ b/Engine/Plugins/Experimental/PythonScriptPlugin/Source/PythonScriptPlugin/Private/PyWrapperObject.h
@@ -242,6 +242,7 @@ private:
 	FPyWrapperObjectMetaData PyMetaData;
 
 	TObjectPtr<UObjectRedirector> ClassRedirector;
+	TObjectPtr<UObjectRedirector> CDORedirector;
 
 	friend class FPythonGeneratedClassBuilder;
 
