Added basic support for deprecated enum entries in Python

Note: This is currently limited to entities deprecated via ScriptName meta-data on the enum entry

diff --git a/Engine/Plugins/Experimental/PythonScriptPlugin/Source/PythonScriptPlugin/Private/PyGenUtil.cpp b/Engine/Plugins/Experimental/PythonScriptPlugin/Source/PythonScriptPlugin/Private/PyGenUtil.cpp
--- a/Engine/Plugins/Experimental/PythonScriptPlugin/Source/PythonScriptPlugin/Private/PyGenUtil.cpp
+++ b/Engine/Plugins/Experimental/PythonScriptPlugin/Source/PythonScriptPlugin/Private/PyGenUtil.cpp
@@ -752,6 +752,11 @@ void FGeneratedWrappedEnumType::Finalize_PostReady()
 				if (PyEnumEntry)
 				{
 					EnumMetaData->EnumEntries.Add(PyEnumEntry);
+
+					for (const FUTF8Buffer& EntryDeprecatedAlias : EnumEntry.EntryDeprecatedAliases)
+					{
+						FPyWrapperEnum::AddDeprecatedEnumEntry(&PyType, PyEnumEntry, EntryDeprecatedAlias.GetData(), EnumEntry.EntryDoc.GetData());
+					}
 				}
 			}
 		}
@@ -786,13 +791,19 @@ void FGeneratedWrappedEnumType::ExtractEnumEntries(const UEnum* InEnum)
 {
 	for (int32 EnumEntryIndex = 0; EnumEntryIndex < InEnum->NumEnums() - 1; ++EnumEntryIndex)
 	{
-		// todo: deprecated enum entries?
 		if (ShouldExportEnumEntry(InEnum, EnumEntryIndex))
 		{
 			FGeneratedWrappedEnumEntry& EnumEntry = EnumEntries.AddDefaulted_GetRef();
 			EnumEntry.EntryName = TCHARToUTF8Buffer(*GetEnumEntryPythonName(InEnum, EnumEntryIndex));
 			EnumEntry.EntryDoc = TCHARToUTF8Buffer(*PythonizeTooltip(ParseTooltip(GetEnumEntryTooltip(InEnum, EnumEntryIndex))));
 			EnumEntry.EntryValue = InEnum->GetValueByIndex(EnumEntryIndex);
+
+			const TArray<FString> DeprecatedEnumEntryNames = GetDeprecatedEnumEntryPythonNames(InEnum, EnumEntryIndex);
+			EnumEntry.EntryDeprecatedAliases.Reserve(DeprecatedEnumEntryNames.Num());
+			for (const FString& DeprecatedEnumEntryName : DeprecatedEnumEntryNames)
+			{
+				EnumEntry.EntryDeprecatedAliases.Add(TCHARToUTF8Buffer(*DeprecatedEnumEntryName));
+			}
 		}
 	}
 }
@@ -3042,6 +3053,50 @@ FString GetEnumEntryPythonName(const UEnum* InEnum, const int32 InEntryIndex)
 	return PythonizeName(EnumEntryName, EPythonizeNameCase::Upper);
 }
 
+TArray<FString> GetDeprecatedEnumEntryPythonNames(const UEnum* InEnum, const int32 InEntryIndex)
+{
+	TArray<FString> EnumEntryNames;
+
+	// First see if we have a name override in the meta-data
+	{
+		const FString EnumEntryName = InEnum->GetMetaData(TEXT("ScriptName"), InEntryIndex);
+
+		// This may be a semi-colon separated list - everything but the first item is deprecated
+		if (!EnumEntryName.IsEmpty())
+		{
+			EnumEntryName.ParseIntoArray(EnumEntryNames, TEXT(";"), false);
+
+			// Remove the non-deprecated entry
+			if (EnumEntryNames.Num() > 0)
+			{
+				EnumEntryNames.RemoveAt(0, EAllowShrinking::No);
+			}
+
+			// Trim whitespace and remove empty items
+			EnumEntryNames.RemoveAll([](FString& InStr)
+			{
+				InStr.TrimStartAndEndInline();
+				return InStr.IsEmpty();
+			});
+
+			for (FString& EnumEntryNamePart : EnumEntryNames)
+			{
+				FText ValidationError;
+				if (!IsValidName(EnumEntryNamePart, &ValidationError))
+				{
+					REPORT_PYTHON_GENERATION_ISSUE(Error, TEXT("'%s' has an invalid 'ScriptName' meta-data value '%s' (from '%s'): %s."), *InEnum->GetPathName(), *EnumEntryNamePart, *EnumEntryName, *ValidationError.ToString());
+				}
+
+				EnumEntryNamePart = PythonizeName(EnumEntryNamePart, EPythonizeNameCase::Upper);
+			}
+		}
+	}
+
+	// TODO: Support reading deprecated names from CoreRedirects?
+
+	return EnumEntryNames;
+}
+
 FString GetDelegatePythonName(const UFunction* InDelegateSignature)
 {
 	return InDelegateSignature->GetName().LeftChop(19); // Trim the "__DelegateSignature" suffix from the name
diff --git a/Engine/Plugins/Experimental/PythonScriptPlugin/Source/PythonScriptPlugin/Private/PyGenUtil.h b/Engine/Plugins/Experimental/PythonScriptPlugin/Source/PythonScriptPlugin/Private/PyGenUtil.h
--- a/Engine/Plugins/Experimental/PythonScriptPlugin/Source/PythonScriptPlugin/Private/PyGenUtil.h
+++ b/Engine/Plugins/Experimental/PythonScriptPlugin/Source/PythonScriptPlugin/Private/PyGenUtil.h
@@ -827,6 +827,9 @@ namespace PyGenUtil
 		FGeneratedWrappedEnumEntry& operator=(FGeneratedWrappedEnumEntry&&) = delete;
 		FGeneratedWrappedEnumEntry& operator=(const FGeneratedWrappedEnumEntry&) = default;
 
+		/** Array of deprecated aliases for this enum entry, if any */
+		TArray<FUTF8Buffer> EntryDeprecatedAliases;
+
 		/** The name of the entry */
 		FUTF8Buffer EntryName;
 
@@ -1192,6 +1195,9 @@ namespace PyGenUtil
 	/** Get the Python name of the given enum entry */
 	FString GetEnumEntryPythonName(const UEnum* InEnum, const int32 InEntryIndex);
 
+	/** Get the deprecated Python names of the given enum entry */
+	TArray<FString> GetDeprecatedEnumEntryPythonNames(const UEnum* InEnum, const int32 InEntryIndex);
+
 	/** Get the Python name of the given delegate signature */
 	FString GetDelegatePythonName(const UFunction* InDelegateSignature);
 
diff --git a/Engine/Plugins/Experimental/PythonScriptPlugin/Source/PythonScriptPlugin/Private/PyWrapperEnum.cpp b/Engine/Plugins/Experimental/PythonScriptPlugin/Source/PythonScriptPlugin/Private/PyWrapperEnum.cpp
--- a/Engine/Plugins/Experimental/PythonScriptPlugin/Source/PythonScriptPlugin/Private/PyWrapperEnum.cpp
+++ b/Engine/Plugins/Experimental/PythonScriptPlugin/Source/PythonScriptPlugin/Private/PyWrapperEnum.cpp
@@ -108,6 +108,9 @@ struct FPyWrapperEnumValueDescrObject
 	/** The enum entry doc string (may be null) */
 	PyObject* EnumEntryDoc;
 
+	/** Set if this enum entry is deprecated and using it should emit a deprecation warning */
+	TOptional<FString> DeprecationMessage;
+
 	/** New an instance */
 	static FPyWrapperEnumValueDescrObject* New(PyTypeObject* InEnumType, const int64 InEnumEntryValue, const char* InEnumEntryName, const char* InEnumEntryDoc)
 	{
@@ -117,6 +120,24 @@ struct FPyWrapperEnumValueDescrObject
 			Self->EnumEntry = FPyWrapperEnum::New(InEnumType);
 			FPyWrapperEnum::Init(Self->EnumEntry, InEnumEntryValue, InEnumEntryName);
 			Self->EnumEntryDoc = InEnumEntryDoc ? PyUnicode_FromString(InEnumEntryDoc) : nullptr;
+			new(&Self->DeprecationMessage) TOptional<FString>();
+		}
+		return Self;
+	}
+
+	/** New a deprecated instance */
+	static FPyWrapperEnumValueDescrObject* NewDeprecated(FPyWrapperEnum* InEnumEntry, const char* InEnumEntryName, const char* InEnumEntryDoc)
+	{
+		FPyWrapperEnumValueDescrObject* Self = (FPyWrapperEnumValueDescrObject*)PyWrapperEnumValueDescrType.tp_alloc(&PyWrapperEnumValueDescrType, 0);
+		if (Self)
+		{
+			Py_XINCREF(InEnumEntry);
+			Self->EnumEntry = InEnumEntry;
+			Self->EnumEntryDoc = InEnumEntryDoc ? PyUnicode_FromString(InEnumEntryDoc) : nullptr;
+			new(&Self->DeprecationMessage) TOptional<FString>(FString::Printf(TEXT("Enum entry '%s.%s' is deprecated: Use '%s.%s'"), 
+				UTF8_TO_TCHAR(Py_TYPE(Self->EnumEntry)->tp_name), UTF8_TO_TCHAR(InEnumEntryName),
+				UTF8_TO_TCHAR(Py_TYPE(Self->EnumEntry)->tp_name), *PyUtil::PyObjectToUEString(Self->EnumEntry->EntryName)
+				));
 		}
 		return Self;
 	}
@@ -130,6 +151,8 @@ struct FPyWrapperEnumValueDescrObject
 		Py_XDECREF(InSelf->EnumEntryDoc);
 		InSelf->EnumEntryDoc = nullptr;
 
+		InSelf->DeprecationMessage.~TOptional<FString>();
+
 		Py_TYPE(InSelf)->tp_free((PyObject*)InSelf);
 	}
 };
@@ -319,6 +342,20 @@ FPyWrapperEnum* FPyWrapperEnum::AddEnumEntry(PyTypeObject* InType, const int64 I
 	return nullptr;
 }
 
+FPyWrapperEnum* FPyWrapperEnum::AddDeprecatedEnumEntry(PyTypeObject* InType, FPyWrapperEnum* InEnumEntry, const char* InEnumEntryName, const char* InEnumEntryDoc)
+{
+	if (!FPyWrapperEnumMetaData::IsEnumFinalized(InType))
+	{
+		FPyWrapperEnumValueDescrObjectPtr PyEnumValueDescr = FPyWrapperEnumValueDescrObjectPtr::StealReference(FPyWrapperEnumValueDescrObject::NewDeprecated(InEnumEntry, InEnumEntryName, InEnumEntryDoc));
+		if (PyEnumValueDescr)
+		{
+			PyDict_SetItemString(InType->tp_dict, InEnumEntryName, (PyObject*)PyEnumValueDescr.GetPtr());
+			return PyEnumValueDescr->EnumEntry;
+		}
+	}
+	return nullptr;
+}
+
 PyTypeObject InitializePyWrapperEnumType()
 {
 	struct FFuncs
@@ -491,6 +528,13 @@ PyTypeObject InitializePyWrapperEnumValueDescrType()
 				}
 			}
 
+			// Deprecated enum entries emit a warning
+			if (InSelf->DeprecationMessage.IsSet() && PyUtil::SetPythonWarning(PyExc_DeprecationWarning, InSelf->EnumEntry, *InSelf->DeprecationMessage.GetValue()) == -1)
+			{
+				// -1 from SetPythonWarning means the warning should be an exception
+				return nullptr;
+			}
+
 			Py_INCREF(InSelf->EnumEntry);
 			return (PyObject*)InSelf->EnumEntry;
 		}
diff --git a/Engine/Plugins/Experimental/PythonScriptPlugin/Source/PythonScriptPlugin/Private/PyWrapperEnum.h b/Engine/Plugins/Experimental/PythonScriptPlugin/Source/PythonScriptPlugin/Private/PyWrapperEnum.h
--- a/Engine/Plugins/Experimental/PythonScriptPlugin/Source/PythonScriptPlugin/Private/PyWrapperEnum.h
+++ b/Engine/Plugins/Experimental/PythonScriptPlugin/Source/PythonScriptPlugin/Private/PyWrapperEnum.h
@@ -60,6 +60,9 @@ struct FPyWrapperEnum : public FPyWrapperBase
 
 	/** Add the given enum entry on the given enum type (returns borrowed reference) */
 	static FPyWrapperEnum* AddEnumEntry(PyTypeObject* InType, const int64 InEnumEntryValue, const char* InEnumEntryName, const char* InEnumEntryDoc);
+
+	/** Add the given deprecated enum entry on the given enum type (returns borrowed reference) */
+	static FPyWrapperEnum* AddDeprecatedEnumEntry(PyTypeObject* InType, FPyWrapperEnum* InEnumEntry, const char* InEnumEntryName, const char* InEnumEntryDoc);
 };
 
 /** Meta-data for all Unreal exposed enum types */
