PythonScriptPlugin: Allow site-packages to differentiate by architecture (arm/x64) as well as platform

diff --git a/Engine/Plugins/Experimental/PythonScriptPlugin/Source/PythonScriptPlugin/Private/PipInstall.cpp b/Engine/Plugins/Experimental/PythonScriptPlugin/Source/PythonScriptPlugin/Private/PipInstall.cpp
--- a/Engine/Plugins/Experimental/PythonScriptPlugin/Source/PythonScriptPlugin/Private/PipInstall.cpp
+++ b/Engine/Plugins/Experimental/PythonScriptPlugin/Source/PythonScriptPlugin/Private/PipInstall.cpp
@@ -93,6 +93,11 @@ FString FPipInstall::WritePluginDependencies(const TArray<TSharedRef<IPlugin>>&
 				continue;
 			}
 
+			if (!CheckCompatibleArch(JsonObj, PyUtil::GetArchString()))
+			{
+				continue;
+			}
+
 			const TArray<TSharedPtr<FJsonValue>>* PyReqs;
 			if (JsonObj->TryGetArrayField(TEXT("Requirements"), PyReqs))
 			{
@@ -444,23 +449,30 @@ void FPipInstall::WriteSitePackagePthFile() const
 		{
 			TSharedRef<IPlugin> Plugin = Plugins[Index];
 			const FString PythonContentPath = FPaths::ConvertRelativePathToFull(Plugin->GetContentDir() / TEXT("Python"));
-			FString PluginPlatformSitePackagesPath = PythonContentPath / TEXT("Lib") / FPlatformMisc::GetUBTPlatform() / TEXT("site-packages");
-			FString PluginGeneralSitePackagesPath = PythonContentPath / TEXT("Lib") / TEXT("site-packages");
+			
+			TArray<TTuple<bool,FString>> SitePackagePathChecks;
+			SitePackagePathChecks.Reserve(PyUtil::GetSitePackageSubdirs().Num());
 
 			// Write platform/general site-packages paths per-plugin to .pth file to account for packaged python dependencies during pip install
-			const bool bPlatformSitePackageExists = FPaths::DirectoryExists(PluginPlatformSitePackagesPath);
-			const bool bGeneralSitePackageExists = FPaths::DirectoryExists(PluginGeneralSitePackagesPath);
+			bool AnySitePackages = false;
+			for (const FString& SitePkgDir : PyUtil::GetSitePackageSubdirs())
+			{
+				FString ChkPath = PythonContentPath / SitePkgDir;
+				bool bPathExists = FPaths::DirectoryExists(ChkPath);
+				
+				SitePackagePathChecks.Emplace(bPathExists, MoveTemp(ChkPath));
+				AnySitePackages |= bPathExists;
+			}
 
-			if (bPlatformSitePackageExists || bGeneralSitePackageExists)
+			if (AnySitePackages)
 			{
 				UE::TUniqueLock ScopeLock(Mutex);
-				if (bPlatformSitePackageExists)
+				for (auto& Check : SitePackagePathChecks)
 				{
-					PluginSitePackagePaths.Emplace(MoveTemp(PluginPlatformSitePackagesPath));
-				}
-				if (bGeneralSitePackageExists)
-				{
-					PluginSitePackagePaths.Emplace(MoveTemp(PluginGeneralSitePackagesPath));
+					if (Check.Get<0>())
+					{
+						PluginSitePackagePaths.Emplace(MoveTemp(Check.Get<1>()));
+					}
 				}
 			}
 		}
@@ -644,10 +656,15 @@ FString FPipInstall::GetVenvInterpreter(const FString& InstallPath)
 bool FPipInstall::CheckCompatiblePlatform(const TSharedPtr<FJsonObject>& JsonObject, const FString& PlatformName)
 {
 	FString JsonPlatform;
-
 	return !JsonObject->TryGetStringField(TEXT("Platform"), JsonPlatform) || JsonPlatform.Equals(TEXT("All"), ESearchCase::IgnoreCase) || JsonPlatform.Equals(PlatformName, ESearchCase::IgnoreCase);
 }
 
+bool FPipInstall::CheckCompatibleArch(const TSharedPtr<FJsonObject>& JsonObject, const FString& ArchName)
+{
+	FString JsonArch;
+	return !JsonObject->TryGetStringField(TEXT("Architecture"), JsonArch) || JsonArch.Equals(TEXT("All"), ESearchCase::IgnoreCase) || JsonArch.Equals(ArchName, ESearchCase::IgnoreCase);
+}
+
 #endif //WITH_PYTHON
 
 #undef LOCTEXT_NAMESPACE
diff --git a/Engine/Plugins/Experimental/PythonScriptPlugin/Source/PythonScriptPlugin/Private/PipInstall.h b/Engine/Plugins/Experimental/PythonScriptPlugin/Source/PythonScriptPlugin/Private/PipInstall.h
--- a/Engine/Plugins/Experimental/PythonScriptPlugin/Source/PythonScriptPlugin/Private/PipInstall.h
+++ b/Engine/Plugins/Experimental/PythonScriptPlugin/Source/PythonScriptPlugin/Private/PipInstall.h
@@ -53,6 +53,8 @@ private:
 
 	static FString GetPythonScriptPluginPath();
 	static FString GetVenvInterpreter(const FString& InstallPath);
+
+	static bool CheckCompatibleArch(const TSharedPtr<FJsonObject>& JsontObject, const FString& ArchName);
 	static bool CheckCompatiblePlatform(const TSharedPtr<FJsonObject>& JsonObject, const FString& PlatformName);
 
 	static int CountInstallLines(const TArray<FString>& RequirementLines);
diff --git a/Engine/Plugins/Experimental/PythonScriptPlugin/Source/PythonScriptPlugin/Private/PyUtil.cpp b/Engine/Plugins/Experimental/PythonScriptPlugin/Source/PythonScriptPlugin/Private/PyUtil.cpp
--- a/Engine/Plugins/Experimental/PythonScriptPlugin/Source/PythonScriptPlugin/Private/PyUtil.cpp
+++ b/Engine/Plugins/Experimental/PythonScriptPlugin/Source/PythonScriptPlugin/Private/PyUtil.cpp
@@ -1093,6 +1093,26 @@ FString GetInterpreterExecutablePath(bool* OutIsEnginePython)
 	return PythonPath;
 }
 
+FString GetArchString()
+{
+#if PLATFORM_CPU_ARM_FAMILY
+	return TEXT("Arm64");
+#else
+	return TEXT("X64");
+#endif //PLATFORM_CPU_ARM_FAMILY
+}
+
+const TArray<FString>& GetSitePackageSubdirs()
+{
+	const static TArray<FString> CheckSitePackageDirs = {
+		FPaths::Combine(TEXT("Lib"), FPlatformMisc::GetUBTPlatform(), PyUtil::GetArchString(), TEXT("site-packages")),
+		FPaths::Combine(TEXT("Lib"), FPlatformMisc::GetUBTPlatform(), TEXT("site-packages")),
+		FPaths::Combine(TEXT("Lib"), TEXT("site-packages")),
+	};
+	return CheckSitePackageDirs;
+}
+	
+
 void AddSitePackagesPath(const FString& InPath)
 {
 	if (!IFileManager::Get().DirectoryExists(*InPath))
diff --git a/Engine/Plugins/Experimental/PythonScriptPlugin/Source/PythonScriptPlugin/Private/PyUtil.h b/Engine/Plugins/Experimental/PythonScriptPlugin/Source/PythonScriptPlugin/Private/PyUtil.h
--- a/Engine/Plugins/Experimental/PythonScriptPlugin/Source/PythonScriptPlugin/Private/PyUtil.h
+++ b/Engine/Plugins/Experimental/PythonScriptPlugin/Source/PythonScriptPlugin/Private/PyUtil.h
@@ -536,6 +536,16 @@ namespace PyUtil
 	 */
 	FString GetInterpreterExecutablePath(bool* OutIsEnginePython = nullptr);
 
+	/**
+	 * Get the arch-specific path component for potential lib site-packages differences (currently X64/Arm) 
+	 */
+	FString GetArchString();
+
+	/**
+	 * Get static list of site-package subdirs to check
+	 */
+	const TArray<FString>& GetSitePackageSubdirs();
+
 	/**
 	 * Register the given path as a known site-packages path (@see site.addsitedir in Python).
 	 */
diff --git a/Engine/Plugins/Experimental/PythonScriptPlugin/Source/PythonScriptPlugin/Private/PythonScriptPlugin.cpp b/Engine/Plugins/Experimental/PythonScriptPlugin/Source/PythonScriptPlugin/Private/PythonScriptPlugin.cpp
--- a/Engine/Plugins/Experimental/PythonScriptPlugin/Source/PythonScriptPlugin/Private/PythonScriptPlugin.cpp
+++ b/Engine/Plugins/Experimental/PythonScriptPlugin/Source/PythonScriptPlugin/Private/PythonScriptPlugin.cpp
@@ -1974,10 +1974,11 @@ void FPythonScriptPlugin::RegisterModulePaths(const FString& InFilesystemPath)
 	{
 		PyUtil::AddSystemPath(PythonContentPath);
 
-		const FString PythonContentPlatformSitePackagesPath = PythonContentPath / TEXT("Lib") / FPlatformMisc::GetUBTPlatform() / TEXT("site-packages");
-		const FString PythonContentGeneralSitePackagesPath = PythonContentPath / TEXT("Lib") / TEXT("site-packages");
-		PyUtil::AddSitePackagesPath(PythonContentPlatformSitePackagesPath);
-		PyUtil::AddSitePackagesPath(PythonContentGeneralSitePackagesPath);
+		for (const FString& SitePkgDir : PyUtil::GetSitePackageSubdirs())
+		{
+			const FString CheckSubdir = PythonContentPath / SitePkgDir;
+			PyUtil::AddSitePackagesPath(CheckSubdir);
+		}
 
 		PyUtil::GetOnDiskUnrealModulesCache().AddModules(*PythonContentPath);
 	}
@@ -1990,10 +1991,11 @@ void FPythonScriptPlugin::UnregisterModulePaths(const FString& InFilesystemPath)
 	const FString PythonContentPath = FPaths::ConvertRelativePathToFull(InFilesystemPath / TEXT("Python"));
 	PyUtil::RemoveSystemPath(PythonContentPath);
 
-	const FString PythonContentPlatformSitePackagesPath = PythonContentPath / TEXT("Lib") / FPlatformMisc::GetUBTPlatform() / TEXT("site-packages");
-	const FString PythonContentGeneralSitePackagesPath = PythonContentPath / TEXT("Lib") / TEXT("site-packages");
-	PyUtil::RemoveSystemPath(PythonContentPlatformSitePackagesPath);
-	PyUtil::RemoveSystemPath(PythonContentGeneralSitePackagesPath);
+	for (const FString& SitePkgDir : PyUtil::GetSitePackageSubdirs())
+	{
+		const FString CheckSubdir = PythonContentPath / SitePkgDir;
+		PyUtil::RemoveSystemPath(CheckSubdir);
+	}
 
 	PyUtil::GetOnDiskUnrealModulesCache().RemoveModules(*PythonContentPath);
 }
