[Python]
- GenerateWrappers is now run in parallel by default

diff --git a/Engine/Plugins/Experimental/PythonScriptPlugin/Source/PythonScriptPlugin/Private/PyWrapperTypeRegistry.cpp b/Engine/Plugins/Experimental/PythonScriptPlugin/Source/PythonScriptPlugin/Private/PyWrapperTypeRegistry.cpp
--- a/Engine/Plugins/Experimental/PythonScriptPlugin/Source/PythonScriptPlugin/Private/PyWrapperTypeRegistry.cpp
+++ b/Engine/Plugins/Experimental/PythonScriptPlugin/Source/PythonScriptPlugin/Private/PyWrapperTypeRegistry.cpp
@@ -53,7 +53,7 @@
 namespace UE::Python::Private
 {
 
-static bool GPythonForceWrapperGenerationOnSingleThread = true;
+static bool GPythonForceWrapperGenerationOnSingleThread = false;
 static FAutoConsoleVariableRef CVarGPythonForceWrapperGenerationOnSingleThread(
 	TEXT("Engine.Python.ForceWrapperGenerationOnSingleThread"),
 	GPythonForceWrapperGenerationOnSingleThread,
