[Python]
- Move one remaining non thread-safe call to game thread
- Implement a way to execute game thread tasks from workers synchronously during wrapper generation

diff --git a/Engine/Plugins/Experimental/PythonScriptPlugin/Source/PythonScriptPlugin/Private/PyGenUtil.cpp b/Engine/Plugins/Experimental/PythonScriptPlugin/Source/PythonScriptPlugin/Private/PyGenUtil.cpp
--- a/Engine/Plugins/Experimental/PythonScriptPlugin/Source/PythonScriptPlugin/Private/PyGenUtil.cpp
+++ b/Engine/Plugins/Experimental/PythonScriptPlugin/Source/PythonScriptPlugin/Private/PyGenUtil.cpp
@@ -23,6 +23,8 @@
 #include "UObject/CoreRedirects.h"
 #include "UObject/PropertyPortFlags.h"
 #include "UObject/UnrealType.h"
+#include "Async/EventCount.h"
+#include "Containers/ConsumeAllMpmcQueue.h"
 #include "Engine/BlueprintCore.h"
 #include "Engine/BlueprintGeneratedClass.h"
 #include "StructUtils/UserDefinedStruct.h"
@@ -84,6 +86,96 @@ const TCHAR* HiddenMetaDataKey = TEXT("Hidden");
 static thread_local TSharedPtr<IBreakIterator> NameBreakIterator;
 static ETypeHintingMode GTypeHintingMode = ETypeHintingMode::Off; // Use through SetTypeHintingMode()/IsTypeHintingEnabled()
 
+namespace MultithreadedGenerationImpl
+{
+	static UE::FEventCount EventCount;
+	static UE::TConsumeAllMpmcQueue<TFunction<void()>> GameThreadTasks;
+}
+
+FMultithreadedGenerationContext::FMultithreadedGenerationContext()
+{
+	check(IsInGameThread());
+	check(GetMultithreadedGenerationRef() == false);
+	GetMultithreadedGenerationRef() = true;
+}
+
+FMultithreadedGenerationContext::~FMultithreadedGenerationContext()
+{
+	check(MultithreadedGenerationImpl::GameThreadTasks.IsEmpty());
+	GetMultithreadedGenerationRef() = false;
+}
+
+bool FMultithreadedGenerationContext::IsMultithreadedGenerationActive()
+{
+	return GetMultithreadedGenerationRef();
+}
+
+void FMultithreadedGenerationContext::EnqueueGameThreadTask(TFunction<void()>&& Task)
+{
+	check(IsMultithreadedGenerationActive());
+	MultithreadedGenerationImpl::GameThreadTasks.ProduceItem(MoveTemp(Task));
+	MultithreadedGenerationImpl::EventCount.Notify();
+}
+
+void FMultithreadedGenerationContext::ExecuteGameThreadTask(TFunctionRef<void()> Task)
+{
+	if (IsInGameThread())
+	{
+		Task();
+	}
+	else
+	{
+		check(IsMultithreadedGenerationActive());
+		UE::FManualResetEvent Event;
+		MultithreadedGenerationImpl::GameThreadTasks.ProduceItem(
+			[&Task, &Event]()
+			{
+				Task();
+				Event.Notify();
+			}
+		);
+		MultithreadedGenerationImpl::EventCount.Notify();
+		Event.Wait();
+	}
+}
+
+void FMultithreadedGenerationContext::ExecuteGameThreadTasksAndWait()
+{
+	bool bWasEmpty = !ExecuteGameThreadTasks();
+
+	if (bWasEmpty)
+	{
+		UE::FEventCountToken Token = MultithreadedGenerationImpl::EventCount.PrepareWait();
+		if (MultithreadedGenerationImpl::GameThreadTasks.IsEmpty())
+		{
+			MultithreadedGenerationImpl::EventCount.Wait(Token);
+		}
+	}
+}
+
+bool FMultithreadedGenerationContext::ExecuteGameThreadTasks()
+{
+	TRACE_CPUPROFILER_EVENT_SCOPE(FMultithreadedGenerationContext::ExecuteGameThreadTasks);
+
+	check(IsInGameThread());
+
+	bool bHadTasks = false;
+	MultithreadedGenerationImpl::GameThreadTasks.ConsumeAllLifo(
+		[&bHadTasks](TFunction<void()> Function)
+		{
+			Function();
+			bHadTasks = true;
+		}
+	);
+
+	return bHadTasks;
+}
+
+bool& FMultithreadedGenerationContext::GetMultithreadedGenerationRef()
+{
+	static bool bMultithreadedGeneration = false;
+	return bMultithreadedGeneration;
+}
 
 void FNativePythonModule::AddType(PyTypeObject* InType)
 {
@@ -2438,12 +2530,16 @@ void PythonizeStructValueImpl(const UScriptStruct* InStruct, const void* InStruc
 					void* SelfArgInstance = SelfParam.ParamProp->ContainerPtrToValuePtr<void>(FuncParams.GetMemory());
 					CastFieldChecked<const FStructProperty>(SelfParam.ParamProp)->Struct->CopyScriptStruct(SelfArgInstance, InStructValue);
 				}
-				{
-					FPyScopedGIL GIL;
-					PyUtil::InvokeFunctionCall(Obj, BreakFuncDef.Func, FuncParams.GetMemory(), TEXT("pythonize default struct value"));
-					PyErr_Clear(); // Clear any errors in case InvokeFunctionCall failed
-				}
 
+				PyGenUtil::FMultithreadedGenerationContext::ExecuteGameThreadTask(
+					[&Obj, &BreakFuncDef, &FuncParams]()
+					{
+						FPyScopedGIL GIL;
+						PyUtil::InvokeFunctionCall(Obj, BreakFuncDef.Func, FuncParams.GetMemory(), TEXT("pythonize default struct value"));
+						PyErr_Clear(); // Clear any errors in case InvokeFunctionCall failed
+					}
+				);
+				
 				// Extract the output argument values as defaults for the struct
 				for (int32 OuputParamIndex = 0; OuputParamIndex < BreakFuncDef.OutputParams.Num(); ++OuputParamIndex)
 				{
diff --git a/Engine/Plugins/Experimental/PythonScriptPlugin/Source/PythonScriptPlugin/Private/PyGenUtil.h b/Engine/Plugins/Experimental/PythonScriptPlugin/Source/PythonScriptPlugin/Private/PyGenUtil.h
--- a/Engine/Plugins/Experimental/PythonScriptPlugin/Source/PythonScriptPlugin/Private/PyGenUtil.h
+++ b/Engine/Plugins/Experimental/PythonScriptPlugin/Source/PythonScriptPlugin/Private/PyGenUtil.h
@@ -86,6 +86,23 @@ namespace PyGenUtil
 		}
 	};
 
+	/** Context used during multithreaded wrapper generation. */
+	struct FMultithreadedGenerationContext
+	{
+		FMultithreadedGenerationContext();
+		~FMultithreadedGenerationContext();
+
+		static bool IsMultithreadedGenerationActive();
+		static void EnqueueGameThreadTask(TFunction<void()>&& Task);
+		static void ExecuteGameThreadTask(TFunctionRef<void()> Task);
+		
+		bool ExecuteGameThreadTasks();
+		void ExecuteGameThreadTasksAndWait();
+
+	private:
+		static bool& GetMultithreadedGenerationRef();
+	};
+
 	/** Stores the information about a native Python module */
 	struct FNativePythonModule
 	{
diff --git a/Engine/Plugins/Experimental/PythonScriptPlugin/Source/PythonScriptPlugin/Private/PyUtil.cpp b/Engine/Plugins/Experimental/PythonScriptPlugin/Source/PythonScriptPlugin/Private/PyUtil.cpp
--- a/Engine/Plugins/Experimental/PythonScriptPlugin/Source/PythonScriptPlugin/Private/PyUtil.cpp
+++ b/Engine/Plugins/Experimental/PythonScriptPlugin/Source/PythonScriptPlugin/Private/PyUtil.cpp
@@ -533,6 +533,9 @@ void ImportDefaultValue(const FProperty* InProp, void* InPropValue, const FStrin
 
 bool InvokeFunctionCall(UObject* InObj, const UFunction* InFunc, void* InBaseParamsAddr, const TCHAR* InErrorCtxt)
 {
+	// FEditorScriptExecutionGuard isn't thread-safe and most function call might also not be thread-safe.
+	check(IsInGameThread());
+
 	bool bThrewException = false;
 	FScopedScriptExceptionHandler ExceptionHandler([InErrorCtxt, &bThrewException](ELogVerbosity::Type Verbosity, const TCHAR* ExceptionMessage, const TCHAR* StackMessage)
 	{
diff --git a/Engine/Plugins/Experimental/PythonScriptPlugin/Source/PythonScriptPlugin/Private/PyWrapperTypeRegistry.cpp b/Engine/Plugins/Experimental/PythonScriptPlugin/Source/PythonScriptPlugin/Private/PyWrapperTypeRegistry.cpp
--- a/Engine/Plugins/Experimental/PythonScriptPlugin/Source/PythonScriptPlugin/Private/PyWrapperTypeRegistry.cpp
+++ b/Engine/Plugins/Experimental/PythonScriptPlugin/Source/PythonScriptPlugin/Private/PyWrapperTypeRegistry.cpp
@@ -956,7 +956,9 @@ void FPyWrapperTypeRegistry::GenerateWrappedTypes(bool bAllowMultithreadedGenera
 		UE::FMutex Mutex;
 		std::atomic<int32> SharedIndex = 0;
 		
+		const uint32 NumWorkers = LowLevelTasks::FScheduler::Get().GetNumWorkers();
 		const bool bIsMultithreadedGenerationAllowed = !UE::Python::Private::GPythonForceWrapperGenerationOnSingleThread && bAllowMultithreadedGeneration;
+		const bool bIsMultithreaded = bIsMultithreadedGenerationAllowed && NumWorkers > 0;
 		TGuardValue GuardIsInMultithreadedGeneration(UE::Python::Private::GIsInMultithreadedGeneration, bIsMultithreadedGenerationAllowed);
 
 		EParallelForFlags ParallelForFlags = EParallelForFlags::Unbalanced;
@@ -965,13 +967,25 @@ void FPyWrapperTypeRegistry::GenerateWrappedTypes(bool bAllowMultithreadedGenera
 			ParallelForFlags |= EParallelForFlags::ForceSingleThread;
 		}
 
+		PyGenUtil::FMultithreadedGenerationContext MultithreadedGenerationContext;
+
 		// Do not try to call EnterProgressFrame on the SlowTask during this because a single FlushRenderingCommand
 		// can take more time than the whole function now that it's parallel.
 		ParallelFor(
-			// No more gains past a certain point because of highly contended PyGIL.
-			FMath::Min(24u, LowLevelTasks::FScheduler::Get().GetNumWorkers()),
-			[this, &Mutex, &SharedIndex, &ObjectsToProcess, &DirtyModules](int32)
+			FMath::Clamp(NumWorkers, 1u, 24u /* No more gains past a certain point because of highly contended PyGIL. */),
+			[this, &Mutex, &SharedIndex, &ObjectsToProcess, &DirtyModules, &bIsMultithreaded, &MultithreadedGenerationContext](int32)
 			{
+				// Under multithreaded scenario, the gamethread is dedicated to running thread-unsafe tasks scheduled by other workers
+				if (IsInGameThread() && bIsMultithreaded)
+				{
+					while (SharedIndex.load(std::memory_order_relaxed) < ObjectsToProcess.Num())
+					{
+						MultithreadedGenerationContext.ExecuteGameThreadTasksAndWait();
+					}
+
+					return;
+				}
+
 				TSet<FName> LocalDirtyModules;
 				FGeneratedWrappedTypeReferences LocalGeneratedWrappedTypeReferences;
 
@@ -980,6 +994,11 @@ void FPyWrapperTypeRegistry::GenerateWrappedTypes(bool bAllowMultithreadedGenera
 					int32 Index = SharedIndex.fetch_add(1, std::memory_order_relaxed);
 					if (Index >= ObjectsToProcess.Num())
 					{
+						if (!IsInGameThread())
+						{
+							// Wake up the game thread
+							PyGenUtil::FMultithreadedGenerationContext::EnqueueGameThreadTask([](){});
+						}
 						break;
 					}
 
@@ -994,7 +1013,7 @@ void FPyWrapperTypeRegistry::GenerateWrappedTypes(bool bAllowMultithreadedGenera
 
 					if (IsInGameThread())
 					{
-						ExecutePostInitOnGameThread();
+						MultithreadedGenerationContext.ExecuteGameThreadTasks();
 					}
 				}
 
@@ -1008,9 +1027,7 @@ void FPyWrapperTypeRegistry::GenerateWrappedTypes(bool bAllowMultithreadedGenera
 
 		PyEval_RestoreThread(ThreadState);
 
-		ExecutePostInitOnGameThread();
-
-		check(RegisteredPostInitForGameThread.IsEmpty());
+		MultithreadedGenerationContext.ExecuteGameThreadTasks();
 	}
 
 	NotifyModulesDirtied(DirtyModules);
@@ -1018,20 +1035,6 @@ void FPyWrapperTypeRegistry::GenerateWrappedTypes(bool bAllowMultithreadedGenera
 	UE_LOG(LogPython, Verbose, TEXT("Took %f seconds to generate and initialize Python wrapped types for the initial load."), GenerateDuration);
 }
 
-void FPyWrapperTypeRegistry::ExecutePostInitOnGameThread()
-{
-	TRACE_CPUPROFILER_EVENT_SCOPE(FPyWrapperTypeRegistry::ExecutePostInitOnGameThread);
-
-	check(IsInGameThread());
-	
-	RegisteredPostInitForGameThread.ConsumeAllLifo(
-		[](TFunction<void()> Function)
-		{
-			Function();
-		}
-	);
-}
-
 void FPyWrapperTypeRegistry::GenerateWrappedTypesForModule(const FName ModuleName)
 {
 	TRACE_CPUPROFILER_EVENT_SCOPE(FPyWrapperTypeRegistry::GenerateWrappedTypesForModule)
@@ -2762,7 +2765,7 @@ PyTypeObject* FPyWrapperTypeRegistry::GenerateWrappedDelegateType(const UFunctio
 		DelegateMetaData->DelegateSignature.SetFunction(InDelegateSignature);
 		GeneratedWrappedType->MetaData = DelegateMetaData;
 
-		RegisteredPostInitForGameThread.ProduceItem(
+		PyGenUtil::FMultithreadedGenerationContext::EnqueueGameThreadTask(
 			[DelegateMetaData, CreatePythonCallableForDelegate_GameThread]()
 			{
 				DelegateMetaData->PythonCallableForDelegateClass = CreatePythonCallableForDelegate_GameThread();
@@ -2778,7 +2781,7 @@ PyTypeObject* FPyWrapperTypeRegistry::GenerateWrappedDelegateType(const UFunctio
 		DelegateMetaData->DelegateSignature.SetFunction(InDelegateSignature);
 		GeneratedWrappedType->MetaData = DelegateMetaData;
 
-		RegisteredPostInitForGameThread.ProduceItem(
+		PyGenUtil::FMultithreadedGenerationContext::EnqueueGameThreadTask(
 			[DelegateMetaData, CreatePythonCallableForDelegate_GameThread]()
 			{
 				DelegateMetaData->PythonCallableForDelegateClass = CreatePythonCallableForDelegate_GameThread();
diff --git a/Engine/Plugins/Experimental/PythonScriptPlugin/Source/PythonScriptPlugin/Private/PyWrapperTypeRegistry.h b/Engine/Plugins/Experimental/PythonScriptPlugin/Source/PythonScriptPlugin/Private/PyWrapperTypeRegistry.h
--- a/Engine/Plugins/Experimental/PythonScriptPlugin/Source/PythonScriptPlugin/Private/PyWrapperTypeRegistry.h
+++ b/Engine/Plugins/Experimental/PythonScriptPlugin/Source/PythonScriptPlugin/Private/PyWrapperTypeRegistry.h
@@ -12,7 +12,6 @@
 #include "Misc/EnumClassFlags.h"
 #include "UObject/GCObject.h"
 #include "UObject/WeakObjectPtr.h"
-#include "Containers/ConsumeAllMpmcQueue.h"
 #include "Containers/StripedMap.h"
 #include "Templates/Function.h"
 
@@ -550,9 +549,6 @@ private:
 	/** Array of information about native Python modules */
 	TArray<PyGenUtil::FNativePythonModule> NativePythonModules;
 
-	/** Functions to be called on the game thread after or during a multi-threaded generation pass. */
-	UE::TConsumeAllMpmcQueue<TFunction<void ()>> RegisteredPostInitForGameThread;
-
 	/** Callback for when a Python module is dirtied */
 	FOnModuleDirtied OnModuleDirtiedDelegate;
 };
