[Python]
- Fix potential deadlock if the game thread loop exited before all workers were done

diff --git a/Engine/Plugins/Experimental/PythonScriptPlugin/Source/PythonScriptPlugin/Private/PyWrapperTypeRegistry.cpp b/Engine/Plugins/Experimental/PythonScriptPlugin/Source/PythonScriptPlugin/Private/PyWrapperTypeRegistry.cpp
--- a/Engine/Plugins/Experimental/PythonScriptPlugin/Source/PythonScriptPlugin/Private/PyWrapperTypeRegistry.cpp
+++ b/Engine/Plugins/Experimental/PythonScriptPlugin/Source/PythonScriptPlugin/Private/PyWrapperTypeRegistry.cpp
@@ -955,6 +955,7 @@ void FPyWrapperTypeRegistry::GenerateWrappedTypes(bool bAllowMultithreadedGenera
 
 		UE::FMutex Mutex;
 		std::atomic<int32> SharedIndex = 0;
+		std::atomic<int32> NumActiveWorkers = 0;
 		
 		const uint32 NumWorkers = LowLevelTasks::FScheduler::Get().GetNumWorkers();
 		const bool bIsMultithreadedGenerationAllowed = !UE::Python::Private::GPythonForceWrapperGenerationOnSingleThread && bAllowMultithreadedGeneration;
@@ -973,12 +974,12 @@ void FPyWrapperTypeRegistry::GenerateWrappedTypes(bool bAllowMultithreadedGenera
 		// can take more time than the whole function now that it's parallel.
 		ParallelFor(
 			FMath::Clamp(NumWorkers, 1u, 24u /* No more gains past a certain point because of highly contended PyGIL. */),
-			[this, &Mutex, &SharedIndex, &ObjectsToProcess, &DirtyModules, &bIsMultithreaded, &MultithreadedGenerationContext](int32)
+			[this, &Mutex, &SharedIndex, &ObjectsToProcess, &DirtyModules, &bIsMultithreaded, &MultithreadedGenerationContext, &NumActiveWorkers](int32)
 			{
 				// Under multithreaded scenario, the gamethread is dedicated to running thread-unsafe tasks scheduled by other workers
 				if (IsInGameThread() && bIsMultithreaded)
 				{
-					while (SharedIndex.load(std::memory_order_relaxed) < ObjectsToProcess.Num())
+					while (SharedIndex.load(std::memory_order_relaxed) < ObjectsToProcess.Num() || NumActiveWorkers.load(std::memory_order_acquire) != 0)
 					{
 						MultithreadedGenerationContext.ExecuteGameThreadTasksAndWait();
 					}
@@ -986,6 +987,8 @@ void FPyWrapperTypeRegistry::GenerateWrappedTypes(bool bAllowMultithreadedGenera
 					return;
 				}
 
+				NumActiveWorkers++;
+
 				TSet<FName> LocalDirtyModules;
 				FGeneratedWrappedTypeReferences LocalGeneratedWrappedTypeReferences;
 
@@ -994,11 +997,6 @@ void FPyWrapperTypeRegistry::GenerateWrappedTypes(bool bAllowMultithreadedGenera
 					int32 Index = SharedIndex.fetch_add(1, std::memory_order_relaxed);
 					if (Index >= ObjectsToProcess.Num())
 					{
-						if (!IsInGameThread())
-						{
-							// Wake up the game thread
-							PyGenUtil::FMultithreadedGenerationContext::EnqueueGameThreadTask([](){});
-						}
 						break;
 					}
 
@@ -1019,6 +1017,14 @@ void FPyWrapperTypeRegistry::GenerateWrappedTypes(bool bAllowMultithreadedGenera
 
 				GenerateWrappedTypesForReferences(LocalGeneratedWrappedTypeReferences, LocalDirtyModules);
 
+				NumActiveWorkers--;
+
+				if (!IsInGameThread())
+				{
+					// Wake up the game thread
+					PyGenUtil::FMultithreadedGenerationContext::EnqueueGameThreadTask([]() {});
+				}
+
 				UE::TUniqueLock ScopeLock(Mutex);
 				DirtyModules.Append(LocalDirtyModules);
 			},
