[Python]
- Fix game thread callback not being called when running the python wrapper generation outside the multithread context (i.e. OnModulesChanged)

diff --git a/Engine/Plugins/Experimental/PythonScriptPlugin/Source/PythonScriptPlugin/Private/PyGenUtil.cpp b/Engine/Plugins/Experimental/PythonScriptPlugin/Source/PythonScriptPlugin/Private/PyGenUtil.cpp
--- a/Engine/Plugins/Experimental/PythonScriptPlugin/Source/PythonScriptPlugin/Private/PyGenUtil.cpp
+++ b/Engine/Plugins/Experimental/PythonScriptPlugin/Source/PythonScriptPlugin/Private/PyGenUtil.cpp
@@ -112,9 +112,16 @@ bool FMultithreadedGenerationContext::IsMultithreadedGenerationActive()
 
 void FMultithreadedGenerationContext::EnqueueGameThreadTask(TFunction<void()>&& Task)
 {
-	check(IsMultithreadedGenerationActive());
-	MultithreadedGenerationImpl::GameThreadTasks.ProduceItem(MoveTemp(Task));
-	MultithreadedGenerationImpl::EventCount.Notify();
+	if (IsInGameThread())
+	{
+		Task();
+	}
+	else
+	{
+		check(IsMultithreadedGenerationActive());
+		MultithreadedGenerationImpl::GameThreadTasks.ProduceItem(MoveTemp(Task));
+		MultithreadedGenerationImpl::EventCount.Notify();
+	}
 }
 
 void FMultithreadedGenerationContext::ExecuteGameThreadTask(TFunctionRef<void()> Task)
