PipInstall: Support on-demand pip installing and provide UI helper for standard background install run

diff --git a/Engine/Plugins/Experimental/PythonScriptPlugin/Source/PythonScriptPlugin/Private/PipInstall.cpp b/Engine/Plugins/Experimental/PythonScriptPlugin/Source/PythonScriptPlugin/Private/PipInstall.cpp
--- a/Engine/Plugins/Experimental/PythonScriptPlugin/Source/PythonScriptPlugin/Private/PipInstall.cpp
+++ b/Engine/Plugins/Experimental/PythonScriptPlugin/Source/PythonScriptPlugin/Private/PipInstall.cpp
@@ -3,6 +3,7 @@
 #include "PipInstall.h"
 #include "PipRunnable.h"
 
+#include "PyGIL.h"
 #include "PyUtil.h"
 #include "PythonScriptPluginSettings.h"
 #include "Dom/JsonObject.h"
@@ -13,10 +14,187 @@
 #include "Misc/FeedbackContext.h"
 #include "Misc/FileHelper.h"
 #include "Misc/Paths.h"
-#include "Misc/ScopedSlowTask.h"
 
 #define LOCTEXT_NAMESPACE "PipInstall"
 
+IPipInstall& IPipInstall::Get()
+{
+	static FPipInstall Instance;
+	return Instance;
+}
+
+
+bool FPipInstall::InitPipInstall()
+{
+	TRACE_CPUPROFILER_EVENT_SCOPE(FPipInstall::InitPipInstaller);
+
+#if WITH_PYTHON
+	// Init Pip installer for python dependencies (if any)
+	if (IsInitialized())
+	{
+		return true;
+	}
+	
+	CheckInvalidPipEnv();
+
+	// Generate the input listing files of plugins with python dependencies and the listing of all requirements (installed or not)
+	TArray<TSharedRef<IPlugin>> PythonPlugins;
+	WritePluginsListing(PythonPlugins);
+
+	for (const TSharedRef<IPlugin>& PyPlugin : PythonPlugins)
+	{
+		// TODO: Get site-package dir listing from PyUtil instead of duplicating
+		// Remove leftover __pycache__ folders from plugins that use pip, but previously used packaged dependencies
+		const FString LibDir = PyPlugin->GetContentDir() / TEXT("Python") / TEXT("Lib");
+		CheckRemoveOrphanedPackages(LibDir / TEXT("site-packages"));
+		CheckRemoveOrphanedPackages(LibDir / FPlatformMisc::GetUBTPlatform() / TEXT("site-packages"));
+	}
+
+	TArray<FString> ReqInLines;
+	TArray<FString> ExtraUrls;
+	const FString ReqsInFile = WritePluginDependencies(PythonPlugins, ReqInLines, ExtraUrls);
+
+	ClearCachedInstallRequirements();
+	if (ReqInLines.IsEmpty())
+	{
+		UE_LOG(LogPython, Display, TEXT("No pip-enabled plugins with python dependencies found, skipping"));
+		// Remove outdated parsed dependency files if there's nothing to install
+		RemoveParsedDependencyFiles();
+		return true;
+	}
+
+	// Some dependencies may need installing
+	FFeedbackContext* Context = GWarn;
+	if (!SetupPipEnv(Context))
+	{
+		return false;
+	}
+
+	if (!CacheDetectInstallDeps(Context))
+	{
+		return false;
+	}
+
+	bInitialized.Store(true);
+	return true;
+#else
+	return false;
+#endif // WITH_PYTHON
+}
+
+bool FPipInstall::LaunchPipInstall(bool RunAsync, TSharedPtr<ICmdProgressNotifier> CmdProgressNotifier)
+{
+	TRACE_CPUPROFILER_EVENT_SCOPE(FPipInstall::LaunchPipInstallAsync);
+
+#if WITH_PYTHON
+	if (!IsInitialized())
+	{
+		// Must explicitly initialize pip installer before use
+		return false;
+	}
+	
+	if (CountCachedInstallRequirements() <= 0)
+	{
+		// Nothing to install
+		return true;
+	}
+
+	FFeedbackContext* Context = GWarn;
+
+	// Run install of all python dependencies for enabled plugins
+	return LaunchPipInstallDependencies(RunAsync, CmdProgressNotifier, Context);
+#else
+	return false;
+#endif
+}
+
+bool FPipInstall::IsInstalling()
+{
+	return IsInitialized() && IsBackgroundInstalling();
+}
+
+int32 FPipInstall::GetNumPackagesToInstall()
+{
+#if WITH_PYTHON
+	if (!IsInitialized())
+	{
+		// TODO: Possibly use TOptional instead
+		return -1;
+	}
+	return CountCachedInstallRequirements();
+#else
+	return -1;
+#endif // WITH_PYTHON
+}
+
+
+bool FPipInstall::GetPackageInstallList(TArray<FString>& PyPackages)
+{
+	PyPackages.Reset();
+#if WITH_PYTHON
+	if (!IsInitialized())
+	{
+		return false;
+	}
+
+	{
+		FScopeLock ScopeLock(&GuardCacheCS);
+		PyPackages = CachedInstallRequirements;
+	}
+	
+	return true;
+#endif // WITH_PYTHON
+}
+
+bool FPipInstall::RegisterPipSitePackagesPath() const
+{
+#if WITH_PYTHON
+	const FString PipSitePackagePath = GetPipSitePackagesPath();
+	{
+		FPyScopedGIL GIL;
+		PyUtil::AddSitePackagesPath(PipSitePackagePath);
+	}
+	return true;
+#else
+	return false;
+#endif //WITH_PYTHON
+}
+
+
+FPipInstall::FPipInstall()
+{
+	TRACE_CPUPROFILER_EVENT_SCOPE(FPipInstall::FPipInstall);
+
+	// TODO: Verify this is necessary
+	ensure(IsInGameThread());
+	
+	bInitialized.Store(false);
+	BackgroundInstallRunnable = nullptr;
+
+	// Default install path: <ProjectDir>/PipInstall
+	PipInstallPath = FPaths::ConvertRelativePathToFull(FPaths::ProjectIntermediateDir() / TEXT("PipInstall"));
+
+	// Check for UE_PIPINSTALL_PATH install path override
+	const FString EnvInstallPath = FPlatformMisc::GetEnvironmentVariable(TEXT("UE_PIPINSTALL_PATH"));
+	if (!EnvInstallPath.IsEmpty())
+	{
+		FText ErrReason;
+		if (FPaths::ValidatePath(EnvInstallPath, &ErrReason))
+		{
+			PipInstallPath = FPaths::ConvertRelativePathToFull(EnvInstallPath);
+		}
+		else
+		{
+			UE_LOG(LogPython, Warning, TEXT("UE_PIPINSTALL_PATH: Invalid path specified: %s"), *ErrReason.ToString());
+		}
+	}
+	
+#if WITH_PYTHON
+	VenvInterp = GetVenvInterpreter(PipInstallPath);
+#endif //WITH_PYTHON
+}
+
+
 #if WITH_PYTHON
 
 // In order to keep editor startup time fast, check directly for this utils version (make sure to match with wheel version in PythonScriptPlugin/Content/Python/Lib/wheels)
@@ -29,21 +207,14 @@ const FString FPipInstall::RequirementsInputFilename = TEXT("merged_requirements
 const FString FPipInstall::ExtraUrlsFilename = TEXT("extra_urls.txt");
 const FString FPipInstall::ParsedRequirementsFilename = TEXT("merged_requirements.txt");
 
-FPipInstall& FPipInstall::Get()
+bool FPipInstall::IsInitialized() const
 {
-	static FPipInstall Instance;
-	return Instance;
+	return bInitialized.Load();
 }
 
-
-bool FPipInstall::IsEnabled() const
+bool FPipInstall::IsBackgroundInstalling() const
 {
-	return bRunOnStartup && !bCmdLineDisable;
-}
-
-bool FPipInstall::IsCmdLineDisabled() const
-{
-	return bCmdLineDisable;
+	return BackgroundInstallRunnable.IsValid() && BackgroundInstallRunnable->IsRunning();
 }
 
 FString FPipInstall::WritePluginsListing(TArray<TSharedRef<IPlugin>>& OutPythonPlugins) const
@@ -243,47 +414,49 @@ void FPipInstall::CheckInvalidPipEnv() const
 	PlatformFile.DeleteDirectoryRecursively(*PipInstallPath);
 }
 
-void FPipInstall::SetupPipEnv(FFeedbackContext* Context, bool bForceRebuild /* = false */) const
+bool FPipInstall::SetupPipEnv(FFeedbackContext* Context, bool bForceRebuild /* = false */) const
 {
 	TRACE_CPUPROFILER_EVENT_SCOPE(FPipInstall::SetupPipEnv);
 
 	const FString EngineInterp = PyUtil::GetInterpreterExecutablePath();
 #ifdef PYTHON_CHECK_SYSEXEC
-	// HACK: Set this compiler variable to check what sys.executable python subproceses get (should match python executable unreal was built against)
+	// HACK: Set this compiler variable to check what sys.executable python subprocesses get (should match python executable unreal was built against)
 	RunPythonCmd(LOCTEXT("PipInstall.DebugInterpWeirdness", "Check Python sys.executable..."), EngineInterp, TEXT("-c \"import sys; print(f'sys.executable: {sys.executable}')\""), Context);
 #endif //PYTHON_CHECK_SYSEXEC
 
 	if (!bForceRebuild && FPaths::FileExists(VenvInterp))
 	{
-		SetupPipInstallUtils(Context);
-		return;
+		return SetupPipInstallUtils(Context);
 	}
 
-	if (bForceRebuild && FPaths::DirectoryExists(PipInstallPath))
+	if (bForceRebuild)
 	{
-		// TODO: Need to cache generated files before deleting the directory (or only delete subdirs)
 		IPlatformFile& PlatformFile = FPlatformFileManager::Get().GetPlatformFile();
 		PlatformFile.DeleteDirectoryRecursively(*PipInstallPath);
+		
+		// Regen the initial files needed here (it's a waste of time, but this is a rare scenario)
+		TArray<TSharedRef<IPlugin>> PythonPlugins;
+		WritePluginsListing(PythonPlugins);
+
+		TArray<FString> ReqInLines;
+		TArray<FString> ExtraUrls;
+		const FString ReqsInFile = WritePluginDependencies(PythonPlugins, ReqInLines, ExtraUrls);
 	}
 
 	const FString VenvCmd = FString::Printf(TEXT("-m venv \"%s\""), *FPaths::ConvertRelativePathToFull(PipInstallPath));
-
-	FScopedSlowTask SubprocTask(0.0f, LOCTEXT("PipInstall.SetupVenv", "Setting up pip install environment..."), true, *Context);
-	SubprocTask.MakeDialog();
-
 	int32 Res = RunPythonCmd(EngineInterp, VenvCmd, Context);
 	if (Res != 0)
 	{
 		UE_LOG(LogPython, Error, TEXT("Unable to create pip install environment (%d)"), Res);
-		return;
+		return false;
 	}
 
-	SetupPipInstallUtils(Context);
+	return SetupPipInstallUtils(Context);
 }
 
 void FPipInstall::RemoveParsedDependencyFiles() const
 {
-	TRACE_CPUPROFILER_EVENT_SCOPE(FPipInstall::ParsePluginDependencies);
+	TRACE_CPUPROFILER_EVENT_SCOPE(FPipInstall::RemoveParsedDependencyFiles);
 
 	const FString ParsedReqsFile = PipInstallPath / ParsedRequirementsFilename;
 	if (FPaths::FileExists(ParsedReqsFile))
@@ -293,7 +466,7 @@ void FPipInstall::RemoveParsedDependencyFiles() const
 	}
 }
 
-FString FPipInstall::ParsePluginDependencies(const FString& MergedInRequirementsFile, FFeedbackContext* Context) const
+bool FPipInstall::ParsePluginDependencies(const FString& MergedInRequirementsFile, FFeedbackContext* Context) const
 {
 	TRACE_CPUPROFILER_EVENT_SCOPE(FPipInstall::ParsePluginDependencies);
 
@@ -307,35 +480,30 @@ FString FPipInstall::ParsePluginDependencies(const FString& MergedInRequirements
 		DisableHashes = TEXT("--disable-hashes");
 	}
 
-	FScopedSlowTask SubprocTask(0.0f, LOCTEXT("PipInstall.ParseRequirements", "Parsing pip requirements..."), true, *Context);
-	SubprocTask.MakeDialog(false,false);
-
 	const FString Cmd = FString::Printf(TEXT("-m ue_parse_plugin_reqs %s -vv \"%s\" \"%s\""), *DisableHashes, *MergedInRequirementsFile, *ParsedReqsFile);
-	RunPythonCmd(VenvInterp, Cmd, Context);
-
-	return FPaths::ConvertRelativePathToFull(ParsedReqsFile);
+	return (RunPythonCmd(VenvInterp, Cmd, Context) == 0);
 }
 
-bool FPipInstall::RunPipInstall(FFeedbackContext* Context) const
+bool FPipInstall::LaunchPipInstallDependencies(bool RunAsync, TSharedPtr<ICmdProgressNotifier> CmdProgressNotifier, FFeedbackContext* Context)
 {
-	TRACE_CPUPROFILER_EVENT_SCOPE(FPipInstall::RunPipInstall);
-
-	const FString ParsedReqsFile = PipInstallPath / ParsedRequirementsFilename;
-	const FString ExtraUrlsFile = PipInstallPath / ExtraUrlsFilename;
+	TRACE_CPUPROFILER_EVENT_SCOPE(FPipInstall::LaunchPipInstallDependencies);
 
-	if (!FPaths::FileExists(ParsedReqsFile))
+	if (IsBackgroundInstalling())
 	{
-		return true;
+		return false;
 	}
 
-	TArray<FString> ParsedReqLines;
-	if ( !FFileHelper::LoadFileToStringArray(ParsedReqLines, *ParsedReqsFile) )
+	const FString ParsedReqsFile = GetParsedRequirementsPath();
+	const FString ExtraUrlsFile = GetExtraUrlsPath();
+
+	if (!FPaths::FileExists(ParsedReqsFile))
 	{
-		return false;
+		UE_LOG(LogPython, Warning, TEXT("PipInstaller background task already running"));
+		return true;
 	}
 
-	int ReqCount = CountInstallLines(ParsedReqLines);
-	if (ReqCount < 1)
+	int ReqCount = CountCachedInstallRequirements();
+	if ( ReqCount == 0 )
 	{
 		return true;
 	}
@@ -347,40 +515,19 @@ bool FPipInstall::RunPipInstall(FFeedbackContext* Context) const
 	}
 
 	const FString Cmd = SetupPipInstallCmd(ParsedReqsFile, ExtraUrls);
-	
-	TSharedRef<ICmdProgressNotifier> PipNotifier = MakeShared<FSlowTaskNotifier>(ReqCount, LOCTEXT("PipInstall.InstallRequirements", "Installing pip requirements..."), Context);
-	TSharedPtr<ICmdProgressParser> ProgParser = MakeShared<FPipProgressParser>(ReqCount, PipNotifier);
 
-	int32 Result = RunPythonCmd(VenvInterp, Cmd, Context, ProgParser);
-	return (Result == 0);
-}
-
-int FPipInstall::NumPackagesToInstall() const
-{
-	const FString ParsedReqsFile = FPaths::ConvertRelativePathToFull(PipInstallPath / ParsedRequirementsFilename);
-
-	TArray<FString> ParsedReqLines;
-	if (!FPaths::FileExists(ParsedReqsFile) || !FFileHelper::LoadFileToStringArray(ParsedReqLines, *ParsedReqsFile))
+	TSharedPtr<ICmdProgressParser> ProgressParser = nullptr;
+	if (CmdProgressNotifier.IsValid())
 	{
-		return 0;
+		ProgressParser = MakeShared<FPipProgressParser>(ReqCount, CmdProgressNotifier.ToSharedRef());
 	}
 
-	return FPipInstall::CountInstallLines(ParsedReqLines);
-}
-
-int FPipInstall::CountInstallLines(const TArray<FString>& RequirementLines)
-{
-	int Count = 0;
-	for (const FStringView Line : RequirementLines)
+	if (RunAsync)
 	{
-		bool bCommentLine = Line.TrimStart().StartsWith(TCHAR('#'));
-		if (!bCommentLine && !Line.Contains(TEXT("# [pkg:check]")))
-		{
-			Count += 1;
-		}
+		return RunPipCmdAsync(VenvInterp, Cmd, Context, ProgressParser);
 	}
 
-	return Count;
+	return RunPipCmdSync(VenvInterp, Cmd, Context, ProgressParser);
 }
 
 FString FPipInstall::GetPipInstallPath() const
@@ -400,36 +547,76 @@ FString FPipInstall::GetPipSitePackagesPath() const
 #endif
 }
 
+bool FPipInstall::CacheDetectInstallDeps(FFeedbackContext* Context)
+{
+	TRACE_CPUPROFILER_EVENT_SCOPE(FPipInstall::CacheDetectInstallDeps);
+	
+	const FString ReqsInFile = GetInputRequirementsPath();
+	if (!ParsePluginDependencies(ReqsInFile, Context))
+	{
+		return false;
+	}
+	
+	return UpdateCachedInstallRequirements();
+}
+
 
-FPipInstall::FPipInstall()
+int32 FPipInstall::CountCachedInstallRequirements()
 {
-	TRACE_CPUPROFILER_EVENT_SCOPE(FPipInstall::FPipInstall);
+	FScopeLock ScopeLock(&GuardCacheCS);
+	return CachedInstallRequirements.Num();
+}
 
-	// Check settings/cmd-line for whether pip installer is enabled
-	bRunOnStartup = GetDefault<UPythonScriptPluginSettings>()->bRunPipInstallOnStartup;
-	bCmdLineDisable = FParse::Param(FCommandLine::Get(), TEXT("DisablePipInstall"));
+void FPipInstall::ClearCachedInstallRequirements()
+{
+	FScopeLock ScopeLock(&GuardCacheCS);
+	CachedInstallRequirements.Empty();
+}
 
-	// Default install path: <ProjectDir>/PipInstall
-	PipInstallPath = FPaths::ConvertRelativePathToFull(FPaths::ProjectIntermediateDir() / TEXT("PipInstall"));
 
-	// Check for UE_PIPINSTALL_PATH install path override
-	const FString EnvInstallPath = FPlatformMisc::GetEnvironmentVariable(TEXT("UE_PIPINSTALL_PATH"));
-	if (!EnvInstallPath.IsEmpty())
+bool FPipInstall::UpdateCachedInstallRequirements()
+{
+	FScopeLock ScopeLock(&GuardCacheCS);
+	CachedInstallRequirements.Empty();
+	
+	const FString ParsedReqsFile = GetParsedRequirementsPath();
+	if (!FPaths::FileExists(ParsedReqsFile))
 	{
-		FText ErrReason;
-		if (FPaths::ValidatePath(EnvInstallPath, &ErrReason))
-		{
-			PipInstallPath = FPaths::ConvertRelativePathToFull(EnvInstallPath);
-		}
-		else
+		return false;
+	}
+
+	TArray<FString> ParsedReqLines;
+	if ( !FFileHelper::LoadFileToStringArray(ParsedReqLines, *ParsedReqsFile) )
+	{
+		return false;
+	}
+
+	for (const FStringView Line : ParsedReqLines)
+	{
+		bool bCommentLine = Line.TrimStart().StartsWith(TCHAR('#'));
+		if (!bCommentLine && !Line.Contains(TEXT("# [pkg:check]")))
 		{
-			UE_LOG(LogPython, Warning, TEXT("UE_PIPINSTALL_PATH: Invalid path specified: %s"), *ErrReason.ToString());
+			CachedInstallRequirements.Add(FString(Line));
 		}
 	}
 
-	VenvInterp = GetVenvInterpreter(PipInstallPath);
+	return true;
+}
+
+FString FPipInstall::GetInputRequirementsPath() const
+{
+	return FPaths::ConvertRelativePathToFull(PipInstallPath / RequirementsInputFilename);
 }
 
+FString FPipInstall::GetParsedRequirementsPath() const
+{
+	return FPaths::ConvertRelativePathToFull(PipInstallPath / ParsedRequirementsFilename);
+}
+
+FString FPipInstall::GetExtraUrlsPath() const
+{
+	return FPaths::ConvertRelativePathToFull(PipInstallPath / ExtraUrlsFilename);
+}
 
 void FPipInstall::WriteSitePackagePthFile() const
 {
@@ -510,19 +697,19 @@ void FPipInstall::WriteSitePackagePthFile() const
 }
 
 
-void FPipInstall::SetupPipInstallUtils(FFeedbackContext* Context) const
+bool FPipInstall::SetupPipInstallUtils(FFeedbackContext* Context) const
 {
 	TRACE_CPUPROFILER_EVENT_SCOPE(FPipInstall::SetupPipInstallUtils);
 
 	if (CheckPipInstallUtils(Context))
 	{
-		return;
+		return true;
 	}
 
 	const FString PythonScriptDir = GetPythonScriptPluginPath();
 	if (PythonScriptDir.IsEmpty())
 	{
-		return;
+		return false;
 	}
 
 	const FString PipWheelsDir = FPaths::ConvertRelativePathToFull(PythonScriptDir / TEXT("Content/Python/Lib/wheels"));
@@ -531,10 +718,7 @@ void FPipInstall::SetupPipInstallUtils(FFeedbackContext* Context) const
 	const FString PipInstallReq = TEXT("ue-pipinstall-utils==") + PipInstallUtilsVer;
 	const FString Cmd = FString::Printf(TEXT("-m pip install --upgrade --no-index --find-links \"%s\" -r \"%s\" %s"), *PipWheelsDir, *InstallRequirements, *PipInstallReq);
 
-	FScopedSlowTask SubprocTask(0.0f, LOCTEXT("PipInstall.SetupPipInstallUtils", "Setting up pip install utils"), true, *Context);
-	SubprocTask.MakeDialog(false,false);
-
-	RunPythonCmd(VenvInterp, Cmd, Context);
+	return (RunPythonCmd(VenvInterp, Cmd, Context) == 0);
 }
 
 
@@ -542,9 +726,6 @@ bool FPipInstall::CheckPipInstallUtils(FFeedbackContext* Context) const
 {
 	// Verify that correct version of pip install utils is already available
 	const FString Cmd = FString::Printf(TEXT("-c \"import pkg_resources;dist=pkg_resources.working_set.find(pkg_resources.Requirement.parse('ue-pipinstall-utils'));exit(dist.version!='%s' if dist is not None else 1)\""), *PipInstallUtilsVer);
-	FScopedSlowTask SubprocTask(0.0f, LOCTEXT("PipInstall.CheckPipInstallUtils", "Check pip install utils installed"), true, *Context);
-	SubprocTask.MakeDialog(false,false);
-
 	return (RunPythonCmd(VenvInterp, Cmd, Context) == 0);
 }
 
@@ -602,6 +783,36 @@ int32 FPipInstall::RunPythonCmd(const FString& PythonInterp, const FString& Cmd,
 	return OutResult;
 }
 
+bool FPipInstall::RunPipCmdAsync(const FString& PythonInterp, const FString& Cmd, FFeedbackContext* Context, TSharedPtr<ICmdProgressParser> CmdParser)
+{
+	TRACE_CPUPROFILER_EVENT_SCOPE(FPipInstall::RunPythonCmd);
+	
+	UE_LOG(LogPython, Log, TEXT("Running background pip install: python %s"), *Cmd);
+	BackgroundInstallRunnable = MakeShared<FLoggedSubprocessThread>(FPaths::ConvertRelativePathToFull(PythonInterp), Cmd, Context, CmdParser);
+	BackgroundInstallRunnable->OnCompleted().BindLambda([this](int32 ReturnCode)
+	{
+		if (ReturnCode == 0)
+		{
+			CacheDetectInstallDeps();
+		}
+	});
+	
+	return BackgroundInstallRunnable->Launch();
+}
+
+bool FPipInstall::RunPipCmdSync(const FString& PythonInterp, const FString& Cmd, FFeedbackContext* Context, TSharedPtr<ICmdProgressParser> CmdParser)
+{
+	int Res = RunPythonCmd(VenvInterp, Cmd, Context, CmdParser);
+	if (Res == 0)
+	{
+		// TODO: Update regardless of success?
+		return CacheDetectInstallDeps();
+	}
+
+	return false;
+}
+
+
 FString FPipInstall::GetPythonScriptPluginPath()
 {
 	TSharedPtr<IPlugin> PythonPlugin = IPluginManager::Get().FindPlugin("PythonScriptPlugin");
diff --git a/Engine/Plugins/Experimental/PythonScriptPlugin/Source/PythonScriptPlugin/Private/PipInstall.h b/Engine/Plugins/Experimental/PythonScriptPlugin/Source/PythonScriptPlugin/Private/PipInstall.h
--- a/Engine/Plugins/Experimental/PythonScriptPlugin/Source/PythonScriptPlugin/Private/PipInstall.h
+++ b/Engine/Plugins/Experimental/PythonScriptPlugin/Source/PythonScriptPlugin/Private/PipInstall.h
@@ -2,25 +2,50 @@
 
 #pragma once
 
+#include "IPipInstall.h"
 #include "Containers/Array.h"
 #include "Containers/UnrealString.h"
+#include "HAL/CriticalSection.h"
 #include "Interfaces/IPluginManager.h"
 #include "Templates/SharedPointer.h"
 
-#if WITH_PYTHON
-
 class FJsonObject;
 class FFeedbackContext;
 
+// Forward declares from PipRunnable.h
 struct ICmdProgressParser;
+class FLoggedSubprocessThread;
 
-class FPipInstall
+// Static singleton implementation of IPipInstall
+class FPipInstall : public IPipInstall
 {
 public:
 	static FPipInstall& Get();
 
-	bool IsEnabled() const;
-	bool IsCmdLineDisabled() const;
+	// IPipInstall Interface Implementation
+	virtual bool InitPipInstall() override;
+	// Run pip to install all missing dependencies for enabled plugins
+	virtual bool LaunchPipInstall(bool RunAsync, TSharedPtr<ICmdProgressNotifier> CmdProgressNotifier = nullptr) override;
+	// Check if a background install is still running
+	virtual bool IsInstalling() override;
+	
+	// Get number of missing python packages to install
+	virtual int32 GetNumPackagesToInstall() override;
+	// Get list of python packages to install
+	virtual bool GetPackageInstallList(TArray<FString>& PyPackages) override;
+	
+	// Register the site-packages path with embedded python env
+	virtual bool RegisterPipSitePackagesPath() const override;
+
+private:
+	friend class IPipInstall;
+	FPipInstall();
+
+#if WITH_PYTHON
+	
+private:
+	bool IsInitialized() const;
+	bool IsBackgroundInstalling() const;
 
 	void CheckRemoveOrphanedPackages(const FString& SitePackagesPath);
 	void CheckInvalidPipEnv() const;
@@ -28,39 +53,48 @@ public:
 	FString WritePluginsListing(TArray<TSharedRef<IPlugin>>& OutPythonPlugins) const;
 	FString WritePluginDependencies(const TArray<TSharedRef<IPlugin>>& PythonPlugins, TArray<FString>& OutRequirements, TArray<FString>& OutExtraUrls) const;
 
-	void SetupPipEnv(FFeedbackContext* Context, bool bForceRebuild = false) const;
+	bool SetupPipEnv(FFeedbackContext* Context, bool bForceRebuild = false) const;
 	void RemoveParsedDependencyFiles() const;
-	FString ParsePluginDependencies(const FString& MergedInRequirementsFile, FFeedbackContext* Context) const;
-	bool RunPipInstall(FFeedbackContext* Context) const;
-
-	int NumPackagesToInstall() const;
+	bool ParsePluginDependencies(const FString& MergedInRequirementsFile, FFeedbackContext* Context) const;
+	bool LaunchPipInstallDependencies(bool RunAsync, TSharedPtr<ICmdProgressNotifier> CmdProgressNotifier, FFeedbackContext* Context);
 
 	FString GetPipInstallPath() const;
 	FString GetPipSitePackagesPath() const;
 
+	bool CacheDetectInstallDeps(FFeedbackContext* Context = nullptr);
+	
+	int32 CountCachedInstallRequirements();
+	void ClearCachedInstallRequirements();
+	bool UpdateCachedInstallRequirements();
+
 private:
-	FPipInstall();
 
+	FString GetInputRequirementsPath() const;
+	FString GetParsedRequirementsPath() const;
+	FString GetExtraUrlsPath() const;
+	
 	void WriteSitePackagePthFile() const;
-	void SetupPipInstallUtils(FFeedbackContext* Context) const;
+	bool SetupPipInstallUtils(FFeedbackContext* Context) const;
 	bool CheckPipInstallUtils(FFeedbackContext* Context) const;
 
+	FString ParseVenvVersion() const;
 	FString SetupPipInstallCmd(const FString& ParsedReqsFile, const TArray<FString>& ExtraUrls) const;
+	bool RunPipCmdAsync(const FString& PythonInterp, const FString& Cmd, FFeedbackContext* Context, TSharedPtr<ICmdProgressParser> CmdParser);
+	bool RunPipCmdSync(const FString& PythonInterp, const FString& Cmd, FFeedbackContext* Context, TSharedPtr<ICmdProgressParser> CmdParser);
 
 	static int32 RunPythonCmd(const FString& PythonInterp, const FString& Cmd, FFeedbackContext* Context, TSharedPtr<ICmdProgressParser> CmdParser = nullptr);
 
-	FString ParseVenvVersion() const;
-
 	static FString GetPythonScriptPluginPath();
 	static FString GetVenvInterpreter(const FString& InstallPath);
 
 	static bool CheckCompatibleArch(const TSharedPtr<FJsonObject>& JsontObject, const FString& ArchName);
 	static bool CheckCompatiblePlatform(const TSharedPtr<FJsonObject>& JsonObject, const FString& PlatformName);
 
-	static int CountInstallLines(const TArray<FString>& RequirementLines);
-
-	bool bRunOnStartup;
-	bool bCmdLineDisable;
+	TAtomic<bool> bInitialized;
+	FCriticalSection GuardCacheCS;
+	TSharedPtr<FLoggedSubprocessThread> BackgroundInstallRunnable;
+	
+	TArray<FString> CachedInstallRequirements;
 
 	FString PipInstallPath;
 	FString VenvInterp;
@@ -71,6 +105,6 @@ private:
 	static const FString RequirementsInputFilename;
 	static const FString ExtraUrlsFilename;
 	static const FString ParsedRequirementsFilename;
-};
-
+	
 #endif //WITH_PYTHON
+};
diff --git a/Engine/Plugins/Experimental/PythonScriptPlugin/Source/PythonScriptPlugin/Private/PipInstallCmdNotifiers.cpp b/Engine/Plugins/Experimental/PythonScriptPlugin/Source/PythonScriptPlugin/Private/PipInstallCmdNotifiers.cpp
new file mode 100644
--- /dev/null
+++ b/Engine/Plugins/Experimental/PythonScriptPlugin/Source/PythonScriptPlugin/Private/PipInstallCmdNotifiers.cpp
@@ -0,0 +1,72 @@
+// Copyright Epic Games, Inc. All Rights Reserved.
+
+#include "PipInstallCmdNotifiers.h"
+#include "PyUtil.h"
+#include "Internationalization/Text.h"
+#include "Misc/AsyncTaskNotification.h"
+#include "Misc/SlowTask.h"
+
+#define LOCTEXT_NAMESPACE "PipInstall"
+
+// FSlowTaskNotifier Implementation
+FSlowTaskNotifier::FSlowTaskNotifier(float GuessSteps, const FText& Description)
+	: SlowTask(MakePimpl<FSlowTask>(GuessSteps, Description, true, *GWarn))
+	, TotalWork(FMath::Max(GuessSteps, 1.0f))
+	, WorkDone(SlowTask->CompletedWork)
+{
+	SlowTask->Initialize();
+	SlowTask->Visibility = ESlowTaskVisibility::Important;
+}
+
+FSlowTaskNotifier::~FSlowTaskNotifier()
+{
+	SlowTask->Destroy();
+}
+
+void FSlowTaskNotifier::UpdateProgress(float UpdateWorkDone, float UpdateTotalWork, const FText& Status)
+{
+	// Exponentially approach 100% if work goes above estimate
+	float NextDone = TotalWork * UpdateWorkDone / UpdateTotalWork;
+	float NextStep = FMath::Max(NextDone - WorkDone, 0.0f);
+	float ProgLeft = FMath::Max(TotalWork - WorkDone, 0.0f);
+
+	// Clamp to 90% of remaining progress
+	float NextWork = FMath::Clamp(NextStep, 0.0f, 0.9*ProgLeft);
+
+	// TODO: Pass in specific requirements to update status lines more accurately
+	SlowTask->EnterProgressFrame(NextWork);
+	WorkDone += NextWork;
+}
+	
+void FSlowTaskNotifier::Completed(bool bSuccess)
+{}
+
+
+//FAsyncTaskCmdNotifier Implementation
+FAsyncTaskCmdNotifier::FAsyncTaskCmdNotifier(float GuessSteps, const FText& Description)
+	: TotalWork(FMath::Max(GuessSteps, 1.0f)), WorkDone(0.0f)
+{
+	FAsyncTaskNotificationConfig NotificationConfig;
+	NotificationConfig.bIsHeadless = false;
+	NotificationConfig.TitleText = Description;
+	NotificationConfig.LogCategory = &LogPython;
+	NotificationConfig.bCanCancel.Set(false);
+	NotificationConfig.bKeepOpenOnSuccess.Set(true);
+	NotificationConfig.bKeepOpenOnFailure.Set(true);
+
+	AsyncNotifier = MakePimpl<FAsyncTaskNotification>(NotificationConfig);
+	AsyncNotifier->SetNotificationState(FAsyncNotificationStateData(Description, FText::GetEmpty(), EAsyncTaskNotificationState::Pending));
+}
+
+void FAsyncTaskCmdNotifier::UpdateProgress(float UpdateWorkDone, float UpdateTotalWork, const FText& Status)
+{
+	float PctDone = (float)UpdateWorkDone / UpdateTotalWork;
+	AsyncNotifier->SetProgressText(FText::Format(LOCTEXT("AsyncTaskCmdNotifier.StatusFmt", "{0} [{1}]"), Status, FText::AsPercent(PctDone)));
+}
+
+void FAsyncTaskCmdNotifier::Completed(bool bSuccess)
+{
+	AsyncNotifier->SetComplete(bSuccess);
+}
+
+#undef LOCTEXT_NAMESPACE
diff --git a/Engine/Plugins/Experimental/PythonScriptPlugin/Source/PythonScriptPlugin/Private/PipInstallCmdNotifiers.h b/Engine/Plugins/Experimental/PythonScriptPlugin/Source/PythonScriptPlugin/Private/PipInstallCmdNotifiers.h
new file mode 100644
--- /dev/null
+++ b/Engine/Plugins/Experimental/PythonScriptPlugin/Source/PythonScriptPlugin/Private/PipInstallCmdNotifiers.h
@@ -0,0 +1,47 @@
+// Copyright Epic Games, Inc. All Rights Reserved.
+
+#pragma once
+
+#include "IPipInstall.h"
+#include "Templates/PimplPtr.h"
+
+class FAsyncTaskNotification;
+struct FSlowTask;
+class FText;
+
+// Slowtask-based notifier for updating command progress (ICmdProgressNotifier)
+class FSlowTaskNotifier : public ICmdProgressNotifier
+{
+public:
+	FSlowTaskNotifier(float GuessSteps, const FText& Description);
+	virtual ~FSlowTaskNotifier() override;
+	
+	// ICmdProgressNotifier methods
+	virtual void UpdateProgress(float UpdateWorkDone, float UpdateTotalWork, const FText& Status) override;
+	virtual void Completed(bool bSuccess) override;
+
+private:
+	TPimplPtr<FSlowTask> SlowTask;
+
+	float TotalWork;
+	float WorkDone;
+};
+
+
+// Asyntask notifier for updating background install commands
+class FAsyncTaskCmdNotifier : public ICmdProgressNotifier
+{
+public:
+	FAsyncTaskCmdNotifier(float GuessSteps, const FText& Description);
+	
+	// ICmdProgressNotifier methods
+	virtual void UpdateProgress(float UpdateWorkDone, float UpdateTotalWork, const FText& Status) override;
+	virtual void Completed(bool bSuccess) override;
+
+private:
+	TPimplPtr<FAsyncTaskNotification> AsyncNotifier;
+
+	float TotalWork;
+	float WorkDone;
+};
+
diff --git a/Engine/Plugins/Experimental/PythonScriptPlugin/Source/PythonScriptPlugin/Private/PipInstallHelpers.cpp b/Engine/Plugins/Experimental/PythonScriptPlugin/Source/PythonScriptPlugin/Private/PipInstallHelpers.cpp
new file mode 100644
--- /dev/null
+++ b/Engine/Plugins/Experimental/PythonScriptPlugin/Source/PythonScriptPlugin/Private/PipInstallHelpers.cpp
@@ -0,0 +1,302 @@
+// Copyright Epic Games, Inc. All Rights Reserved.
+
+#include "PipInstallHelpers.h"
+#include "IPipInstall.h"
+#include "PipInstallLauncher.h"
+#include "PipInstallCmdNotifiers.h"
+#include "PyUtil.h"
+
+#include "Framework/Application/SlateApplication.h"
+#include "Interfaces/IMainFrameModule.h"
+#include "Internationalization/Text.h"
+#include "Widgets/SWindow.h"
+#include "Widgets/Input/SButton.h"
+#include "Widgets/Layout/SUniformGridPanel.h"
+#include "Widgets/Text/STextBlock.h"
+#include "Widgets/Views/STableRow.h"
+#include "Widgets/Views/SListView.h"
+
+#define LOCTEXT_NAMESPACE "PipInstall"
+
+class SPipInstallDialogWidget : public SCompoundWidget
+{
+public:
+	SLATE_BEGIN_ARGS(SPipInstallDialogWidget)
+		: _PipInstallPtr(nullptr)
+		, _AllowBackground(true)
+		, _DialogWindow(nullptr)
+		{}
+
+		SLATE_ARGUMENT(IPipInstall*, PipInstallPtr)
+		SLATE_ARGUMENT(bool, AllowBackground)
+		SLATE_ARGUMENT(TSharedPtr<SWindow>, DialogWindow)
+	SLATE_END_ARGS()
+
+public:
+	void Construct(const FArguments& Args)
+	{
+		PipInstallPtr = Args._PipInstallPtr;
+		DialogWindow = Args._DialogWindow;
+		AllowBackground = Args._AllowBackground;
+
+		InitInsallList();
+
+		FText InstallButtonTooltip;
+		FText CancelButtonTooltip;
+		if (AllowBackground)
+		{
+			InstallButtonTooltip = LOCTEXT("PipInstallUI.InstallBG.ToolTip", "Click to begin background package installation");
+			CancelButtonTooltip = LOCTEXT("PipInstallUI.CancelBG.ToolTip", "Click to begin background package installation");
+		}
+		else
+		{
+			InstallButtonTooltip = LOCTEXT("PipInstallUI.InstallFG.ToolTip", "Click to begin package installation");
+			CancelButtonTooltip = LOCTEXT("PipInstallUI.CancelFG.ToolTip", "Packages must be installed to continue");
+		}
+		
+		this->ChildSlot
+		[
+			SNew(SVerticalBox)
+			+ SVerticalBox::Slot()
+				.AutoHeight()
+				.Padding(2)
+			[
+				SNew(STextBlock)
+				.Text(FText::FormatOrdered(LOCTEXT("PipInstallPackages.Title", "Packages to Install ({0}): "), PipInstallPtr->GetNumPackagesToInstall()))
+			]
+			+ SVerticalBox::Slot()
+				.AutoHeight()
+				.MaxHeight(200.0)
+				.Padding(2)
+			[
+				SAssignNew(PackageInstallList, SListView<TSharedPtr<FString>>)
+				.SelectionMode(ESelectionMode::None)
+				.ListItemsSource(&InstallPackages)
+				.OnGenerateRow(this, &SPipInstallDialogWidget::OnGenerateListRow)
+			]
+			+ SVerticalBox::Slot()
+				.AutoHeight()
+				.Padding(2)
+			[
+				SNew(SUniformGridPanel)
+				+ SUniformGridPanel::Slot(1, 0)
+				[
+					SNew(SButton)
+						.HAlign(HAlign_Center)
+						.Text(LOCTEXT("PipInstallUI.Install", "Install Packages"))
+						.ToolTipText(InstallButtonTooltip)
+						.OnClicked(this, &SPipInstallDialogWidget::OnInstallClicked)
+				]
+				+ SUniformGridPanel::Slot(2, 0)
+				[
+					SNew(SButton)
+						.HAlign(HAlign_Center)
+						.Text(LOCTEXT("PipInstallUI.Cancel", "Cancel"))
+						.IsEnabled(AllowBackground && !PipInstallPtr->IsInstalling())
+						.ToolTipText(CancelButtonTooltip)
+						.OnClicked(this, &SPipInstallDialogWidget::OnCancelClicked)
+				]
+			]
+		];
+	}
+
+	//~ Begin SWidget interface
+	virtual bool SupportsKeyboardFocus() const override { return true; }
+	virtual FReply OnKeyDown(const FGeometry& MyGeometry, const FKeyEvent& InKeyEvent) override
+	{
+		if (InKeyEvent.GetKey() == EKeys::Escape && AllowBackground)
+		{
+			return OnCancelClicked();
+		}
+		else if ( InKeyEvent.GetKey() == EKeys::Enter  )
+		{
+			return OnInstallClicked();
+		}
+
+		return FReply::Unhandled();
+	}
+	//~ End SWidget interface
+
+	bool ShouldStartInstall()
+	{
+		return StartInstall;
+	}
+	
+private:
+	void InitInsallList()
+	{
+		TArray<FString> Packages;
+		PipInstallPtr->GetPackageInstallList(Packages);
+		for ( const FString& Package : Packages )
+		{
+			// Remove hashes/comments, etc from the pip requirements lines keep only "package==ver" portion
+			FStringView PackageView = Package;
+			PackageView.TrimStartInline();
+			int32 CutIdx;
+			if (PackageView.FindChar(TEXT(' '), CutIdx))
+			{
+				PackageView.LeftInline(CutIdx);
+			}
+
+			InstallPackages.Add(MakeShared<FString>(PackageView));
+		}
+	}
+
+	TSharedRef<ITableRow> OnGenerateListRow(TSharedPtr<FString> InItem, const TSharedRef<STableViewBase>& OwnerTable)
+	{
+		return SNew(STableRow<TSharedPtr<FString>>, OwnerTable)
+			[
+				SNew(SOverlay)
+				.Visibility(EVisibility::SelfHitTestInvisible)
+
+				+SOverlay::Slot()
+				.Padding(0.0f)
+				[
+					SNew(SImage)
+					.ColorAndOpacity(FLinearColor::Black)
+				]
+				+SOverlay::Slot()
+				.Padding(0.0f)
+				[
+					SNew(STextBlock)
+					.Text(FText::FromString(*InItem.Get()))
+				]	
+			];
+	}
+
+	FReply OnInstallClicked()
+	{
+		StartInstall = true;
+		if (DialogWindow.IsValid())
+		{
+			DialogWindow.Pin()->RequestDestroyWindow();
+		}
+		
+		return FReply::Handled();
+	}
+
+	FReply OnCancelClicked()
+	{
+		StartInstall = false;
+		if (DialogWindow.IsValid())
+		{
+			DialogWindow.Pin()->RequestDestroyWindow();
+		}
+		
+		return FReply::Handled();
+	}
+	
+private:
+	bool StartInstall = false;
+	
+	bool AllowBackground = true;
+	// Main settings object
+	IPipInstall* PipInstallPtr = nullptr;
+	// Main window
+	TWeakPtr<SWindow> DialogWindow;
+	
+	// List view for packages needing install
+	TSharedPtr<SListView<TSharedPtr<FString>>> PackageInstallList;
+
+	// List of packages to install
+	TArray<TSharedPtr<FString>> InstallPackages;
+};
+
+bool ShowPackageInstallDialogModal(IPipInstall& PipInstall, bool bAllowBackground)
+{
+	TSharedPtr<SWindow> ParentWindow = nullptr;
+	if (FModuleManager::Get().IsModuleLoaded("MainFrame"))
+	{
+		IMainFrameModule& MainFrame = FModuleManager::LoadModuleChecked<IMainFrameModule>("MainFrame");
+		ParentWindow = MainFrame.GetParentWindow();
+	}
+	
+	TSharedRef<SWindow> Window = SNew(SWindow)
+		.Title(LOCTEXT("PipInstallDialog.Title", "Python Dependencies Install"))
+		.SizingRule(ESizingRule::Autosized)
+		.AutoCenter(EAutoCenter::PrimaryWorkArea);
+
+	TSharedPtr<SPipInstallDialogWidget> PipInstallDialog;
+	Window->SetContent(
+		SAssignNew(PipInstallDialog, SPipInstallDialogWidget)
+			.PipInstallPtr(&PipInstall)
+			.AllowBackground(bAllowBackground)
+			.DialogWindow(Window)
+	);
+
+	FSlateApplication::Get().AddModalWindow(Window, ParentWindow, /*bSlowTaskWindow =*/false);
+
+	return PipInstallDialog->ShouldStartInstall();
+}
+
+int FPipInstallHelper::GetNumPackagesToInstall()
+{
+	check(IsInGameThread());
+	
+	IPipInstall& PipInstall = IPipInstall::Get();
+	if (!PipInstall.InitPipInstall())
+	{
+		UE_LOG(LogPython, Error, TEXT("Unable to initialize Pip Install"));
+		return -1;
+	}
+
+	return PipInstall.GetNumPackagesToInstall();
+}
+
+
+EPipInstallDialogResult FPipInstallHelper::ShowPipInstallDialog(bool bAllowBackgroundInstall)
+{
+	check(IsInGameThread());
+	
+	IPipInstall& PipInstall = IPipInstall::Get();
+	if (!PipInstall.InitPipInstall())
+	{
+		UE_LOG(LogPython, Error, TEXT("Unable to initialize Pip Install"));
+		return EPipInstallDialogResult::Error;
+	}
+	
+	bool bClickedInstall = ShowPackageInstallDialogModal(PipInstall, bAllowBackgroundInstall);
+	
+	// Install must be run immediately (game-thread) if background disallowed
+	if (!bAllowBackgroundInstall)
+	{
+		return (PipInstallLauncher::StartSync(PipInstall) ? EPipInstallDialogResult::Finished : EPipInstallDialogResult::Error); 
+	}
+	else if (bClickedInstall && !PipInstall.IsInstalling())
+	{
+		return (PipInstallLauncher::StartAsync(PipInstall) ? EPipInstallDialogResult::BackgroundInstall : EPipInstallDialogResult::Error);
+	}
+	
+	return EPipInstallDialogResult::Canceled;
+}
+
+bool FPipInstallHelper::LaunchHeadlessPipInstall()
+{
+	check(IsInGameThread());
+
+	IPipInstall& PipInstall = IPipInstall::Get();
+	if (!PipInstall.InitPipInstall())
+	{
+		UE_LOG(LogPython, Error, TEXT("Unable to initialize Pip Install"));
+		return false;
+	}
+
+	return PipInstallLauncher::StartSync(PipInstall); 
+}
+
+namespace PipInstallLauncher
+{
+	bool StartSync(IPipInstall& PipInstall)
+	{
+		TSharedPtr<ICmdProgressNotifier> CmdNotifier = MakeShared<FSlowTaskNotifier>(PipInstall.GetNumPackagesToInstall(), LOCTEXT("PipInstall.FGInstallText", "Installing Python Dependencies..."));
+		return PipInstall.LaunchPipInstall(false, CmdNotifier);
+	}
+
+	bool StartAsync(IPipInstall& PipInstall)
+	{
+		TSharedPtr<ICmdProgressNotifier> CmdProgressNotifier = MakeShared<FAsyncTaskCmdNotifier>(PipInstall.GetNumPackagesToInstall(), LOCTEXT("PipInstall.BGInstallText", "Installing Python Dependencies..."));
+		return PipInstall.LaunchPipInstall(true, CmdProgressNotifier);
+	}
+}
+
+#undef LOCTEXT_NAMESPACE
diff --git a/Engine/Plugins/Experimental/PythonScriptPlugin/Source/PythonScriptPlugin/Private/PipInstallLauncher.h b/Engine/Plugins/Experimental/PythonScriptPlugin/Source/PythonScriptPlugin/Private/PipInstallLauncher.h
new file mode 100644
--- /dev/null
+++ b/Engine/Plugins/Experimental/PythonScriptPlugin/Source/PythonScriptPlugin/Private/PipInstallLauncher.h
@@ -0,0 +1,11 @@
+// Copyright Epic Games, Inc. All Rights Reserved.
+
+#pragma once
+
+#include "IPipInstall.h"
+
+namespace PipInstallLauncher
+{
+	bool StartSync(IPipInstall& PipInstall);
+	bool StartAsync(IPipInstall& PipInstall);
+}
diff --git a/Engine/Plugins/Experimental/PythonScriptPlugin/Source/PythonScriptPlugin/Private/PipRunnable.cpp b/Engine/Plugins/Experimental/PythonScriptPlugin/Source/PythonScriptPlugin/Private/PipRunnable.cpp
--- a/Engine/Plugins/Experimental/PythonScriptPlugin/Source/PythonScriptPlugin/Private/PipRunnable.cpp
+++ b/Engine/Plugins/Experimental/PythonScriptPlugin/Source/PythonScriptPlugin/Private/PipRunnable.cpp
@@ -3,16 +3,14 @@
 #include "PipRunnable.h"
 
 #include "PyUtil.h"
-
-#include "Misc/AsyncTaskNotification.h"
 #include "Misc/FeedbackContext.h"
-#include "Misc/SlowTask.h"
 
 #define LOCTEXT_NAMESPACE "PipInstall"
 
-#if WITH_PYTHON
-
 // FPipProgressParser Implementation
+const TArray<FString> FPipProgressParser::MatchStatusStrs = {TEXT("Requirement"), TEXT("Downloading"), TEXT("Collecting"), TEXT("Using"), TEXT("Installing")};
+const TMap<FString,FString> FPipProgressParser::LogReplaceStrs = {{TEXT("Installing collected packages:"), TEXT("Installing collected python package dependencies:")}};
+
 FPipProgressParser::FPipProgressParser(int GuessRequirementsCount, TSharedRef<ICmdProgressNotifier> InCmdNotifier)
 	: RequirementsDone(0.0f)
 	, RequirementsCount(FMath::Max(2.0f*GuessRequirementsCount,1.0f))
@@ -78,44 +76,11 @@ FString FPipProgressParser::ReplaceUpdateStrs(const FString& Line)
 	return RepLine;
 }
 
-
-const TArray<FString> FPipProgressParser::MatchStatusStrs = {TEXT("Requirement"), TEXT("Downloading"), TEXT("Collecting"), TEXT("Using"), TEXT("Installing")};
-const TMap<FString,FString> FPipProgressParser::LogReplaceStrs = {{TEXT("Installing collected packages:"), TEXT("Installing collected python package dependencies:")}};
-
-
-// FSlowTaskNotifier implementation
-FSlowTaskNotifier::FSlowTaskNotifier(float GuessSteps, const FText& Description, FFeedbackContext* Context)
-	: SlowTask(nullptr)
-	, TotalWork(FMath::Max(GuessSteps, 1.0f))
-	, WorkDone(0.0f)
-{
-	SlowTask = MakeUnique<FSlowTask>(GuessSteps, Description, true, *Context);
-	WorkDone = SlowTask->CompletedWork;
-}
-
-void FSlowTaskNotifier::UpdateProgress(float UpdateWorkDone, float UpdateTotalWork, const FText& Status)
-{
-	// Exponentially approach 100% if work goes above estimate
-	float NextDone = TotalWork * UpdateWorkDone / UpdateTotalWork;
-	float NextStep = FMath::Max(NextDone - WorkDone, 0.0f);
-	float ProgLeft = FMath::Max(TotalWork - WorkDone, 0.0f);
-
-	// Clamp to 90% of remaining progress
-	float NextWork = FMath::Clamp(NextStep, 0.0f, 0.9*ProgLeft);
-
-	// TODO: Pass in specific requirements to update status lines more accurately
-	SlowTask->EnterProgressFrame(NextWork, Status);
-	WorkDone += NextWork;
-}
-
-void FSlowTaskNotifier::Completed(bool bSuccess){}
-
-
-// Implementation of a tickable class for running a subprocess with status updates
-class FRunLoggedSubprocess
+// Implementation of a tickable class for running a subprocess with status updates via ICmdProgressParser interface
+class FLoggedSubprocess
 {
 public:
-	FRunLoggedSubprocess(const FString& URI, const FString& Params, FFeedbackContext* Context, TSharedPtr<ICmdProgressParser> CmdParser)
+	FLoggedSubprocess(const FString& URI, const FString& Params, FFeedbackContext* Context, TSharedPtr<ICmdProgressParser> CmdParser)
 		: Context(Context), CmdParser(CmdParser)
 	{
 		bLogOutput = (Context != nullptr);
@@ -129,17 +94,17 @@ public:
 		ProcessHandle = FPlatformProcess::CreateProc(*URI, *Params, false, true, true, nullptr, 0, nullptr, StdOutPipeWrite, nullptr);
 	}
 
-	~FRunLoggedSubprocess()
+	~FLoggedSubprocess()
 	{
 		FPlatformProcess::ClosePipe(StdOutPipeRead, StdOutPipeWrite);
 	}
 
-	bool IsValid()
+	bool IsValid() const
 	{
 		return ProcessHandle.IsValid();
 	}
 
-	int32 GetExitCode()
+	int32 GetExitCode() const
 	{
 		return ProcExitCode;
 	}
@@ -147,7 +112,26 @@ public:
 	bool TickSubprocess()
 	{
 		bool bProcessFinished = FPlatformProcess::GetProcReturnCode(ProcessHandle, &ProcExitCode);
-		BufferedText += FPlatformProcess::ReadPipe(StdOutPipeRead);
+		ParseOutput(FPlatformProcess::ReadPipe(StdOutPipeRead));
+
+		return !bProcessFinished;
+	}
+
+	void KillSubprocess()
+	{
+		FPlatformProcess::TerminateProc(ProcessHandle, true);
+
+		// Make sure to read the final bytes from the pipe before closing
+		ParseOutput(FPlatformProcess::ReadPipe(StdOutPipeRead));
+
+		// Indicate failure in the return code
+		ProcExitCode = -1;
+	}
+
+private:
+	void ParseOutput(const FString& NewOutput)
+	{
+		BufferedText += NewOutput;
 
 		int32 EndOfLineIdx;
 		while (BufferedText.FindChar(TEXT('\n'), EndOfLineIdx))
@@ -164,11 +148,8 @@ public:
 
 			BufferedText.MidInline(EndOfLineIdx + 1, MAX_int32, EAllowShrinking::No);
 		}
-
-		return !bProcessFinished;
 	}
-
-private:
+	
 	FFeedbackContext* Context;
 	TSharedPtr<ICmdProgressParser> CmdParser;
 
@@ -184,7 +165,7 @@ private:
 
 bool FLoggedSubprocessSync::Run(int32& OutExitCode, const FString& URI, const FString& Params, FFeedbackContext* Context, TSharedPtr<ICmdProgressParser> CmdParser)
 {
-	FRunLoggedSubprocess Subproc(URI, Params, Context, CmdParser);
+	FLoggedSubprocess Subproc(URI, Params, Context, CmdParser);
 	if (!Subproc.IsValid())
 	{
 		return false;
@@ -199,6 +180,35 @@ bool FLoggedSubprocessSync::Run(int32& OutExitCode, const FString& URI, const FS
 	return true;
 }
 
-#endif //WITH_PYTHON
+// FLoggedSubprocessThread implementation
+FLoggedSubprocessThread::FLoggedSubprocessThread(const FString& URI, const FString& Params, FFeedbackContext* Context, TSharedPtr<ICmdProgressParser> CmdParser)
+	: Process(URI, Params, true, true)
+	, Context(Context)
+	, CmdParser(CmdParser)
+	, bLogOutput(Context != nullptr)
+{
+	Process.OnOutput().BindRaw(this, &FLoggedSubprocessThread::ParseOutput);
+	Process.OnCompleted().BindRaw(this, &FLoggedSubprocessThread::HandleCompleted);
+}
+
+void FLoggedSubprocessThread::ParseOutput(FString StdoutLine) const
+{
+	bool bOnlyStatusUpdate = CmdParser.IsValid() && CmdParser->UpdateStatus(StdoutLine);
+	if ( bLogOutput && !bOnlyStatusUpdate)
+	{
+		Context->Log(LogPython.GetCategoryName(), ELogVerbosity::Log, StdoutLine);
+	}
+}
+
+void FLoggedSubprocessThread::HandleCompleted(int32 ReturnCode) const
+{
+	if (CmdParser.IsValid())
+	{
+		CmdParser->NotifyCompleted(ReturnCode==0);
+	}
+	
+	OnProcCompleted.ExecuteIfBound(ReturnCode);
+}
+
 
 #undef LOCTEXT_NAMESPACE
diff --git a/Engine/Plugins/Experimental/PythonScriptPlugin/Source/PythonScriptPlugin/Private/PipRunnable.h b/Engine/Plugins/Experimental/PythonScriptPlugin/Source/PythonScriptPlugin/Private/PipRunnable.h
--- a/Engine/Plugins/Experimental/PythonScriptPlugin/Source/PythonScriptPlugin/Private/PipRunnable.h
+++ b/Engine/Plugins/Experimental/PythonScriptPlugin/Source/PythonScriptPlugin/Private/PipRunnable.h
@@ -2,12 +2,11 @@
 
 #pragma once
 
-#if WITH_PYTHON
-
+#include "IPipInstall.h"
 #include "Containers/UnrealString.h"
 #include "HAL/Runnable.h"
 #include "HAL/RunnableThread.h"
-#include "Misc/AsyncTaskNotification.h"
+#include "Misc/MonitoredProcess.h"
 #include "Templates/SharedPointer.h"
 
 struct FSlowTask;
@@ -25,14 +24,6 @@ struct ICmdProgressParser
 	virtual void NotifyCompleted(bool bSuccess) = 0;
 };
 
-// Interface to wrap notifications for progress updates (e.g. slowtask or asyncnotify)
-struct ICmdProgressNotifier
-{
-    virtual ~ICmdProgressNotifier(){};
-	virtual void UpdateProgress(float UpdateWorkDone, float UpdateTotalWork, const FText& Status) = 0;
-	virtual void Completed(bool bSuccess) = 0;
-};
-
 // Pip progress parser implemenation of ICmdProgressParser
 class FPipProgressParser : public ICmdProgressParser
 {
@@ -56,26 +47,32 @@ private:
 	static const TMap<FString,FString> LogReplaceStrs;
 };
 
-// Slowtask-based notifier for updating command progress (ICmdProgressNotifier)
-class FSlowTaskNotifier : public ICmdProgressNotifier
+class FLoggedSubprocessSync
 {
 public:
-	FSlowTaskNotifier(float GuessSteps, const FText& Description, FFeedbackContext* Context);
-    // ICmdProgressNotifier methods
-	virtual void UpdateProgress(float UpdateWorkDone, float UpdateTotalWork, const FText& Status) override;
-	virtual void Completed(bool bSuccess) override;
-
-private:
-	TUniquePtr<FSlowTask> SlowTask;
-
-	float TotalWork;
-	float WorkDone;
+	static bool Run(int32& OutExitCode, const FString& URI, const FString& Params, FFeedbackContext* Context, TSharedPtr<ICmdProgressParser> CmdParser);
 };
 
-// Run a subprocess synchronously (assumes running in game thread)
-struct FLoggedSubprocessSync
+// Threaded subprocess runner (converted to wrapped FMonitoredProcess)
+DECLARE_DELEGATE_OneParam(FOnSubprocessThreadCompleted, int32)
+class FLoggedSubprocessThread
 {
-    static bool Run(int32& OutExitCode, const FString& URI, const FString& Params, FFeedbackContext* Context, TSharedPtr<ICmdProgressParser> CmdParser);
-};
+public:
+	FLoggedSubprocessThread(const FString& URI, const FString& Params, FFeedbackContext* Context, TSharedPtr<ICmdProgressParser> CmdParser);
 
-#endif //WITH_PYTHON
+	bool Launch() {return Process.Launch();}
+	bool IsRunning() {return Process.Update();}
+	
+	FOnSubprocessThreadCompleted& OnCompleted() {return OnProcCompleted;}
+
+private:
+	void ParseOutput(FString StdoutLine) const;
+	void HandleCompleted(int32 ReturnCode) const;
+	
+	FOnSubprocessThreadCompleted OnProcCompleted;
+	
+	FMonitoredProcess Process;
+	FFeedbackContext* Context;
+	TSharedPtr<ICmdProgressParser> CmdParser;
+	const bool bLogOutput;
+};
diff --git a/Engine/Plugins/Experimental/PythonScriptPlugin/Source/PythonScriptPlugin/Private/PythonScriptPlugin.cpp b/Engine/Plugins/Experimental/PythonScriptPlugin/Source/PythonScriptPlugin/Private/PythonScriptPlugin.cpp
--- a/Engine/Plugins/Experimental/PythonScriptPlugin/Source/PythonScriptPlugin/Private/PythonScriptPlugin.cpp
+++ b/Engine/Plugins/Experimental/PythonScriptPlugin/Source/PythonScriptPlugin/Private/PythonScriptPlugin.cpp
@@ -59,7 +59,8 @@
 #include "ContentBrowserFileDataSource.h"
 #include "Toolkits/GlobalEditorCommonCommands.h"
 #include "Misc/FeedbackContext.h"
-#include "PipInstall.h"
+#include "IPipInstall.h"
+#include "PipInstallLauncher.h"
 #endif	// WITH_EDITOR
 
 #if PLATFORM_WINDOWS
@@ -740,6 +741,7 @@ bool FPythonScriptPlugin::ForceEnablePythonAtRuntime(const FString& Location)
 #endif
 }
 
+
 bool FPythonScriptPlugin::ExecPythonCommand(const TCHAR* InPythonCommand)
 {
 	FPythonCommandEx PythonCommand;
@@ -1250,14 +1252,14 @@ void FPythonScriptPlugin::InitializePython()
 		FModuleManager::Get().OnModulesChanged().AddRaw(this, &FPythonScriptPlugin::OnModulesChanged);
 
 #if WITH_EDITOR
+		IPipInstall& PipInstaller = IPipInstall::Get();
 		UE::Tasks::FTask PipTask = 
 			UE::Tasks::Launch(
 				TEXT("PipInstallAsync"),
-				[this]()
+				[&PipInstaller]()
 				{
 					// Init PipInstall task
-					InitPipInstaller();
-					RunPipInstaller();
+					PipInstaller.InitPipInstall();
 				}
 			);
 #endif // WITH_EDITOR
@@ -1270,14 +1272,22 @@ void FPythonScriptPlugin::InitializePython()
 
 #if WITH_EDITOR
 		PipTask.Wait();
-#endif
 
-		// Add Pip UBT install path to site-packages if it exists
-		const FString PipSitePackagePath = FPipInstall::Get().GetPipSitePackagesPath();
-		if (FPaths::DirectoryExists(PipSitePackagePath))
+		// Always add pip install site-package path since actual installs may be deferred
+		PipInstaller.RegisterPipSitePackagesPath();
+
+		// Check settings/cmd-line for whether pip installer is forced to run before initialization finished
+		bool bPipInstallOnInit = GetDefault<UPythonScriptPluginSettings>()->bRunPipInstallOnStartup;
+		bool bCmdLineForceInstall = FParse::Param(FCommandLine::Get(), TEXT("ForcePipInstallOnInit"));
+		bool bForceInstall = bPipInstallOnInit || bCmdLineForceInstall;
+		const int32 NumPkgs = PipInstaller.GetNumPackagesToInstall();
+		if (bForceInstall && NumPkgs > 0)
 		{
-			PyUtil::AddSitePackagesPath(PipSitePackagePath);
+			// Run PipInstall (in game thread) if any of the force options are on
+			PipInstallLauncher::StartSync(PipInstaller);
 		}
+		
+#endif
 
 		// Initialize the tick handler
 		TickHandle = FTSTicker::GetCoreTicker().AddTicker(FTickerDelegate::CreateLambda([this](float DeltaTime)
@@ -1426,92 +1436,6 @@ void FPythonScriptPlugin::ShutdownPython()
 	bHasTicked = false;
 }
 
-void FPythonScriptPlugin::InitPipInstaller()
-{
-	TRACE_CPUPROFILER_EVENT_SCOPE(FPythonScriptPlugin::InitPipInstaller);
-
-	// Init Pip installer for python dependencies (if any)
-	FFeedbackContext* Context = GWarn;
-	FScopedSlowTask PipInstallTask(0, LOCTEXT("PipInstall.RunInit", "Running Pip Init Tasks..."), true, *Context);
-
-	FPipInstall& PipInstaller = FPipInstall::Get();
-	PipInstaller.CheckInvalidPipEnv();
-
-	// Generate the input listing files of plugins with python dependencies and the listing of all requirements (installed or not)
-	TArray<TSharedRef<IPlugin>> PythonPlugins;
-	PipInstaller.WritePluginsListing(PythonPlugins);
-
-	for (const TSharedRef<IPlugin>& PyPlugin : PythonPlugins)
-	{
-		// Remove leftover __pycache__ folders from plugins that use pip, but previously used packaged dependencies
-		const FString LibDir = PyPlugin->GetContentDir() / TEXT("Python") / TEXT("Lib");
-		PipInstaller.CheckRemoveOrphanedPackages(LibDir / TEXT("site-packages"));
-		PipInstaller.CheckRemoveOrphanedPackages(LibDir / FPlatformMisc::GetUBTPlatform() / TEXT("site-packages"));
-	}
-
-	TArray<FString> ReqInLines;
-	TArray<FString> ExtraUrls;
-	const FString ReqsInFile = PipInstaller.WritePluginDependencies(PythonPlugins, ReqInLines, ExtraUrls);
-
-	bool EnabledOnStart = PipInstaller.IsEnabled();
-	if (ReqInLines.IsEmpty())
-	{
-		UE_CLOG(EnabledOnStart, LogPython, Display, TEXT("No enabled plugins with python dependencies found, skipping"));
-		// Remove older parsed dependency files if there's nothing to install
-		PipInstaller.RemoveParsedDependencyFiles();
-		return;
-	}
-
-	// Just return immediately with warning if some python dependencies exist and pip install is disabled
-	const FString PipSitePackagePath = FPaths::ConvertRelativePathToFull(PipInstaller.GetPipSitePackagesPath());
-	if (!EnabledOnStart)
-	{
-		if (PipInstaller.IsCmdLineDisabled() || GIsBuildMachine)
-		{
-			// Don't warn if disabled on cmd-line or is build process
-			UE_LOG(LogPython, Display, TEXT("Enabled plugins have python dependencies, install manually to: %s"), *PipSitePackagePath);
-			UE_LOG(LogPython, Display, TEXT("  See package requirements: %s"), *ReqsInFile);
-		}
-		else
-		{
-			UE_LOG(LogPython, Warning, TEXT("Enabled plugins have python dependencies, enable 'Run Pip Install On Startup' or install manually to: %s"), *PipSitePackagePath);
-			UE_LOG(LogPython, Warning, TEXT("  See package requirements: %s"), *ReqsInFile);
-		}
-		return;
-	}
-
-	UE_LOG(LogPython, Display, TEXT("Preparing to install python dependencies into: %s"), *PipSitePackagePath);
-
-	PipInstaller.SetupPipEnv(Context);
-	PipInstaller.ParsePluginDependencies(ReqsInFile, Context);
-}
-
-void FPythonScriptPlugin::RunPipInstaller()
-{
-	TRACE_CPUPROFILER_EVENT_SCOPE(FPythonScriptPlugin::RunPipInstaller);
-
-	FPipInstall& PipInstaller = FPipInstall::Get();
-	if (!PipInstaller.IsEnabled())
-	{
-		return;
-	}
-
-	if (PipInstaller.NumPackagesToInstall() == 0)
-	{
-		// Nothing to install
-		return;
-	}
-
-	FFeedbackContext* Context = GWarn;
-
-	// Run install of all python dependencies for enabled plugins
-	if (!PipInstaller.RunPipInstall(Context))
-	{
-		UE_LOG(LogPython, Warning, TEXT("Unable to install plugin python dependencies"));
-		return;
-	}
-}
-
 void FPythonScriptPlugin::RequestStubCodeGeneration()
 {
 	// Ignore requests made before the startup scripts have ran
diff --git a/Engine/Plugins/Experimental/PythonScriptPlugin/Source/PythonScriptPlugin/Private/PythonScriptPluginSettings.h b/Engine/Plugins/Experimental/PythonScriptPlugin/Source/PythonScriptPlugin/Private/PythonScriptPluginSettings.h
--- a/Engine/Plugins/Experimental/PythonScriptPlugin/Source/PythonScriptPlugin/Private/PythonScriptPluginSettings.h
+++ b/Engine/Plugins/Experimental/PythonScriptPlugin/Source/PythonScriptPlugin/Private/PythonScriptPluginSettings.h
@@ -91,11 +91,10 @@ public:
 	bool bDeveloperMode;
 
 	/**
-	 * Should the pip install task be run on editor startup.
+	 * Should the pip install task be run on python engine initialization.
 	 *
 	 * NOTE: The project pip install directory: <ProjectDir>/Intermediate/PipInstall/Lib/site-packages
-	 *       will still be added to site package path, if it exists, to allow for pre-populated installs
-	 *       even auto-running UBT PipInstall is disabled.
+	 *       will still be added to site package path, to allow for pre-populated or deferred installs.
 	 *
 	 *       See <ProjectDir>/Intermediate/PipInstall/merged_requirements.in for listing of required packages.
 	 */
@@ -108,7 +107,7 @@ public:
 	UPROPERTY(config, EditAnywhere, Category=PythonPipInstall, AdvancedDisplay, meta=(ConfigRestartRequired=true))
 	bool bPipStrictHashCheck = true;
 
-	/** Only run pip to generate merged requiremnts and/or validate requested packages are installed */
+	/** Only run pip to generate merged requirements and/or validate requested packages are installed */
 	UPROPERTY(config, EditAnywhere, Category=PythonPipInstall, AdvancedDisplay, meta=(ConfigRestartRequired=true))
 	bool bOfflineOnly = false;
 
diff --git a/Engine/Plugins/Experimental/PythonScriptPlugin/Source/PythonScriptPlugin/Public/IPipInstall.h b/Engine/Plugins/Experimental/PythonScriptPlugin/Source/PythonScriptPlugin/Public/IPipInstall.h
new file mode 100644
--- /dev/null
+++ b/Engine/Plugins/Experimental/PythonScriptPlugin/Source/PythonScriptPlugin/Public/IPipInstall.h
@@ -0,0 +1,44 @@
+// Copyright Epic Games, Inc. All Rights Reserved.
+
+#pragma once
+
+#include "Containers/Array.h"
+#include "Internationalization/Text.h"
+#include "Templates/SharedPointer.h"
+
+/**
+ * Interface for wrapping a notifier that can return info from pip install
+ */
+struct ICmdProgressNotifier
+{
+	virtual void UpdateProgress(float UpdateWorkDone, float UpdateTotalWork, const FText& Status) = 0;
+	virtual void Completed(bool bSuccess) = 0;
+
+	virtual ~ICmdProgressNotifier(){};
+};
+
+/**
+ * Interface to pip installer used by python script plugin for installing plugin python dependencies
+ */
+class IPipInstall
+{
+public:
+	static PYTHONSCRIPTPLUGIN_API IPipInstall& Get();
+	
+	// Initialize the internal pip virtual env and check enabled plugins for python dependencies
+	virtual bool InitPipInstall() = 0;
+	// Run pip to install all missing python dependencies for enabled plugins
+	virtual bool LaunchPipInstall(bool RunAsync, TSharedPtr<ICmdProgressNotifier> CmdProgressNotifier = nullptr) = 0;
+	// Check if a background install is running
+	virtual bool IsInstalling() = 0;
+	
+	// Get number of missing python packages to install
+	virtual int32 GetNumPackagesToInstall() = 0;
+	// Get list of python packages to install
+	virtual bool GetPackageInstallList(TArray<FString>& PyPackages) = 0;
+
+	// Register the site-packages path with embedded python env
+	virtual bool RegisterPipSitePackagesPath() const = 0;
+	
+	virtual ~IPipInstall(){};
+};
diff --git a/Engine/Plugins/Experimental/PythonScriptPlugin/Source/PythonScriptPlugin/Public/PipInstallHelpers.h b/Engine/Plugins/Experimental/PythonScriptPlugin/Source/PythonScriptPlugin/Public/PipInstallHelpers.h
new file mode 100644
--- /dev/null
+++ b/Engine/Plugins/Experimental/PythonScriptPlugin/Source/PythonScriptPlugin/Public/PipInstallHelpers.h
@@ -0,0 +1,30 @@
+// Copyright Epic Games, Inc. All Rights Reserved.
+
+#pragma once
+
+#include "IPipInstall.h"
+
+// PipInstall dialog responses
+enum EPipInstallDialogResult
+{
+	BackgroundInstall,
+	Finished,
+	Canceled,
+	Error,
+};
+
+/**
+ * Simplified UI helper for creating a modal dialog and running (background) installs
+ */
+class PYTHONSCRIPTPLUGIN_API FPipInstallHelper
+{
+public:
+	// Check if pip install is required
+	static int32 GetNumPackagesToInstall();
+	
+	// Show notification that pip install is required
+	static EPipInstallDialogResult ShowPipInstallDialog(bool bAllowBackgroundInstall = true);
+
+	// Run a headless pip install (for commandlets/build machines)
+	static bool LaunchHeadlessPipInstall();
+};
