Exposed unreal.collect_garbage() to Python

Unlike unreal.SystemLibrary.collect_garbage(), this runs immediately (rather than the next frame) and will handle being called from a commandlet correctly

diff --git a/Engine/Plugins/Experimental/PythonScriptPlugin/Source/PythonScriptPlugin/Private/PyCore.cpp b/Engine/Plugins/Experimental/PythonScriptPlugin/Source/PythonScriptPlugin/Private/PyCore.cpp
--- a/Engine/Plugins/Experimental/PythonScriptPlugin/Source/PythonScriptPlugin/Private/PyCore.cpp
+++ b/Engine/Plugins/Experimental/PythonScriptPlugin/Source/PythonScriptPlugin/Private/PyCore.cpp
@@ -1663,6 +1663,15 @@ PyObject* FlushGeneratedTypeReinstancing(PyObject* InSelf)
 	Py_RETURN_NONE;
 }
 
+PyObject* CollectGarbage(PyObject* InSelf)
+{
+	Py_BEGIN_ALLOW_THREADS
+	::CollectGarbage(IsRunningCommandlet() ? RF_NoFlags : GARBAGE_COLLECTION_KEEPFLAGS);
+	Py_END_ALLOW_THREADS
+
+	Py_RETURN_NONE;
+}
+
 PyObject* GetTypeFromClass(PyObject* InSelf, PyObject* InArgs)
 {
 	PyObject* PyObj = nullptr;
@@ -1923,6 +1932,7 @@ PyMethodDef PyCoreMethods[] = {
 	{ "generate_struct", PyCFunctionCast(&GenerateStruct), METH_VARARGS, "generate_struct(struct_type: type) -> None -- generate an Unreal struct for the given Python type" },
 	{ "generate_enum", PyCFunctionCast(&GenerateEnum), METH_VARARGS, "generate_enum(enum_type: type) -> None -- generate an Unreal enum for the given Python type" },
 	{ "flush_generated_type_reinstancing", PyCFunctionCast(&FlushGeneratedTypeReinstancing), METH_NOARGS, "flush_generated_type_reinstancing() -> None -- flush any pending reinstancing requests for Python generated types" },
+	{ "collect_garbage", PyCFunctionCast(&CollectGarbage), METH_NOARGS, "collect_garbage() -> None -- run Unreal garbage collection" },
 	{ "get_type_from_class", PyCFunctionCast(&GetTypeFromClass), METH_VARARGS, "get_type_from_class(class_: Class) -> type -- get the best matching Python type for the given Unreal class" },
 	{ "get_type_from_struct", PyCFunctionCast(&GetTypeFromStruct), METH_VARARGS, "get_type_from_struct(struct: Struct) -> type -- get the best matching Python type for the given Unreal struct" },
 	{ "get_type_from_enum", PyCFunctionCast(&GetTypeFromEnum), METH_VARARGS, "get_type_from_enum(enum: Enum) -> type -- get the best matching Python type for the given Unreal enum" },
