[PythonWrapper]
- Fix hang with -onethread because of ParallelFor having its own set of rules to go wide

diff --git a/Engine/Plugins/Experimental/PythonScriptPlugin/Source/PythonScriptPlugin/Private/PyWrapperTypeRegistry.cpp b/Engine/Plugins/Experimental/PythonScriptPlugin/Source/PythonScriptPlugin/Private/PyWrapperTypeRegistry.cpp
--- a/Engine/Plugins/Experimental/PythonScriptPlugin/Source/PythonScriptPlugin/Private/PyWrapperTypeRegistry.cpp
+++ b/Engine/Plugins/Experimental/PythonScriptPlugin/Source/PythonScriptPlugin/Private/PyWrapperTypeRegistry.cpp
@@ -41,6 +41,7 @@
 #include "Misc/CoreDelegates.h"
 #include "Async/UniqueLock.h"
 #include "Async/ParallelFor.h"
+#include "Tasks/Task.h"
 #include "Sanitizer/RaceDetector.h"
 #if WITH_EDITOR
 #include "Kismet2/ReloadUtilities.h"
@@ -962,31 +963,11 @@ void FPyWrapperTypeRegistry::GenerateWrappedTypes(bool bAllowMultithreadedGenera
 		const bool bIsMultithreaded = bIsMultithreadedGenerationAllowed && NumWorkers > 0;
 		TGuardValue GuardIsInMultithreadedGeneration(UE::Python::Private::GIsInMultithreadedGeneration, bIsMultithreadedGenerationAllowed);
 
-		EParallelForFlags ParallelForFlags = EParallelForFlags::Unbalanced;
-		if (!bIsMultithreadedGenerationAllowed)
-		{
-			ParallelForFlags |= EParallelForFlags::ForceSingleThread;
-		}
-
 		PyGenUtil::FMultithreadedGenerationContext MultithreadedGenerationContext;
 
-		// Do not try to call EnterProgressFrame on the SlowTask during this because a single FlushRenderingCommand
-		// can take more time than the whole function now that it's parallel.
-		ParallelFor(
-			FMath::Clamp(NumWorkers, 1u, 24u /* No more gains past a certain point because of highly contended PyGIL. */),
-			[this, &Mutex, &SharedIndex, &ObjectsToProcess, &DirtyModules, &bIsMultithreaded, &MultithreadedGenerationContext, &NumActiveWorkers](int32)
+		auto Worker = 
+			[this, &Mutex, &SharedIndex, &ObjectsToProcess, &DirtyModules, &bIsMultithreaded, &MultithreadedGenerationContext, &NumActiveWorkers]()
 			{
-				// Under multithreaded scenario, the gamethread is dedicated to running thread-unsafe tasks scheduled by other workers
-				if (IsInGameThread() && bIsMultithreaded)
-				{
-					while (SharedIndex.load(std::memory_order_relaxed) < ObjectsToProcess.Num() || NumActiveWorkers.load(std::memory_order_acquire) != 0)
-					{
-						MultithreadedGenerationContext.ExecuteGameThreadTasksAndWait();
-					}
-
-					return;
-				}
-
 				NumActiveWorkers++;
 
 				TSet<FName> LocalDirtyModules;
@@ -1027,9 +1008,38 @@ void FPyWrapperTypeRegistry::GenerateWrappedTypes(bool bAllowMultithreadedGenera
 
 				UE::TUniqueLock ScopeLock(Mutex);
 				DirtyModules.Append(LocalDirtyModules);
-			},
-			ParallelForFlags
-		);
+			};
+
+		if (bIsMultithreaded)
+		{
+			UE::Tasks::FTask WorkerTask =
+				UE::Tasks::Launch(
+					TEXT("GeneratingPythonTypesWorker"),
+					[&NumWorkers, &Worker]() 
+					{
+						ParallelFor(
+							FMath::Clamp(NumWorkers, 1u, 24u /* No more gains past a certain point because of highly contended PyGIL. */),
+							[&Worker](int32)
+							{		
+								Worker(); 
+							},
+							EParallelForFlags::Unbalanced
+						);
+					}
+				);
+
+			// Under multithreaded scenario, the gamethread is dedicated to running thread-unsafe tasks scheduled by other workers
+			while (SharedIndex.load(std::memory_order_relaxed) < ObjectsToProcess.Num() || NumActiveWorkers.load(std::memory_order_acquire) != 0)
+			{
+				MultithreadedGenerationContext.ExecuteGameThreadTasksAndWait();
+			}
+
+			WorkerTask.Wait();
+		}
+		else
+		{
+			Worker();
+		}
 
 		PyEval_RestoreThread(ThreadState);
 
