Encoding Fix for PipInstall

In `FPipInstall::WriteSitePackagePthFile`, the encoding used by default to write the `.pth` file with `FFileHelper::SaveStringArrayToFile` ends up being `Unicode` for Mac and `ANSI` or Windows.
It is not the same as the encoding used by Python to decode the same file. The code in `site.py > site.addsitedir > site.addpackage` uses the `locale` encoding which ends up being `US-ASCII` for Mac and `cp1252` or Windows.
The `Unicode` file starts with a BOM with the first byte being `0xFF` on a Mac, and this byte cannot be decoded by the `ascii` codec which ends up putting Python in an error state, creating the issue we are seeing of `PyObject_CallObject` returning null`.

To avoid this, we force `FFileHelper::SaveStringArrayToFile` to use `UTF-8` without a BOM, which seems to be properly read in both Mac and Windows. The locale encoding of Python is still not matching the one of the file, so this might still create issues.

diff --git a/Engine/Plugins/Experimental/PythonScriptPlugin/Source/PythonScriptPlugin/Private/PipInstall.cpp b/Engine/Plugins/Experimental/PythonScriptPlugin/Source/PythonScriptPlugin/Private/PipInstall.cpp
--- a/Engine/Plugins/Experimental/PythonScriptPlugin/Source/PythonScriptPlugin/Private/PipInstall.cpp
+++ b/Engine/Plugins/Experimental/PythonScriptPlugin/Source/PythonScriptPlugin/Private/PipInstall.cpp
@@ -699,7 +699,7 @@ void FPipInstall::WriteSitePackagePthFile() const
 
 	// Create .pth file in PipInstall/Lib/site-packages to account for plugins with packaged dependencies
 	const FString PyPluginsSitePackageFile = FPaths::ConvertRelativePathToFull(GetPipSitePackagesPath() / PluginsSitePackageFilename);
-	FFileHelper::SaveStringArrayToFile(PluginSitePackagePaths, *PyPluginsSitePackageFile);
+	FFileHelper::SaveStringArrayToFile(PluginSitePackagePaths, *PyPluginsSitePackageFile, FFileHelper::EEncodingOptions::ForceUTF8WithoutBOM);
 }
 
 
