Avoid crashing when executing Python from a thread

#rnx

diff --git a/Engine/Plugins/Experimental/PythonScriptPlugin/Source/PythonScriptPlugin/Private/PyUtil.cpp b/Engine/Plugins/Experimental/PythonScriptPlugin/Source/PythonScriptPlugin/Private/PyUtil.cpp
--- a/Engine/Plugins/Experimental/PythonScriptPlugin/Source/PythonScriptPlugin/Private/PyUtil.cpp
+++ b/Engine/Plugins/Experimental/PythonScriptPlugin/Source/PythonScriptPlugin/Private/PyUtil.cpp
@@ -533,8 +533,15 @@ void ImportDefaultValue(const FProperty* InProp, void* InPropValue, const FStrin
 
 bool InvokeFunctionCall(UObject* InObj, const UFunction* InFunc, void* InBaseParamsAddr, const TCHAR* InErrorCtxt)
 {
-	// FEditorScriptExecutionGuard isn't thread-safe and most function call might also not be thread-safe.
-	check(IsInGameThread());
+	// FEditorScriptExecutionGuard isn't thread-safe and most function calls might also not be thread-safe.
+	// Python editor code that creates threads and interact with the Unreal API will always hit this case 
+	// Instead of crashing, raise a Python Exception to alert the user
+	if (!IsInGameThread())
+	{
+		FPyScopedGIL GIL;
+		SetPythonError(PyExc_RuntimeError, InErrorCtxt, TEXT("Attempted to access Unreal API from outside the main game thread."));
+		return false;
+	}
 
 	bool bThrewException = false;
 	FScopedScriptExceptionHandler ExceptionHandler([InErrorCtxt, &bThrewException](ELogVerbosity::Type Verbosity, const TCHAR* ExceptionMessage, const TCHAR* StackMessage)
