[Python] Update Python docs generator to include enterprise as part of engine

diff --git a/Engine/Plugins/Experimental/PythonScriptPlugin/Source/PythonScriptPlugin/Private/PyOnlineDocsWriter.h b/Engine/Plugins/Experimental/PythonScriptPlugin/Source/PythonScriptPlugin/Private/PyOnlineDocsWriter.h
--- a/Engine/Plugins/Experimental/PythonScriptPlugin/Source/PythonScriptPlugin/Private/PyOnlineDocsWriter.h
+++ b/Engine/Plugins/Experimental/PythonScriptPlugin/Source/PythonScriptPlugin/Private/PyOnlineDocsWriter.h
@@ -14,10 +14,10 @@ enum class EPyOnlineDocsFilterFlags : uint8
 {
 	IncludeNone = 0,
 	IncludeEngine = 1<<0,
-	IncludeEnterprise = 1<<1,
+	IncludeEnterprise UE_DEPRECATED(5.6, "Enterprise is now considered a part of engine, please update usage to the engine enum flag.") = IncludeEngine,
 	IncludeInternal = 1<<2,
 	IncludeProject = 1<<3,
-	IncludeAll = IncludeEngine | IncludeEnterprise | IncludeInternal | IncludeProject,
+	IncludeAll = IncludeEngine | IncludeInternal | IncludeProject,
 };
 ENUM_CLASS_FLAGS(EPyOnlineDocsFilterFlags);
 
diff --git a/Engine/Plugins/Experimental/PythonScriptPlugin/Source/PythonScriptPlugin/Private/PyWrapperTypeRegistry.cpp b/Engine/Plugins/Experimental/PythonScriptPlugin/Source/PythonScriptPlugin/Private/PyWrapperTypeRegistry.cpp
--- a/Engine/Plugins/Experimental/PythonScriptPlugin/Source/PythonScriptPlugin/Private/PyWrapperTypeRegistry.cpp
+++ b/Engine/Plugins/Experimental/PythonScriptPlugin/Source/PythonScriptPlugin/Private/PyWrapperTypeRegistry.cpp
@@ -3172,7 +3172,7 @@ void FPyWrapperTypeRegistry::GenerateStubCodeForWrappedTypes(const EPyOnlineDocs
 					else if (ModulePath.Contains(TEXT("/Enterprise/")))
 					{
 						// Optionally exclude Enterprise classes
-						if (!EnumHasAnyFlags(InDocGenFlags, EPyOnlineDocsFilterFlags::IncludeEnterprise))
+						if (!EnumHasAnyFlags(InDocGenFlags, EPyOnlineDocsFilterFlags::IncludeEngine))
 						{
 							return;
 						}
diff --git a/Engine/Plugins/Experimental/PythonScriptPlugin/Source/PythonScriptPlugin/Private/PythonOnlineDocsCommandlet.cpp b/Engine/Plugins/Experimental/PythonScriptPlugin/Source/PythonScriptPlugin/Private/PythonOnlineDocsCommandlet.cpp
--- a/Engine/Plugins/Experimental/PythonScriptPlugin/Source/PythonScriptPlugin/Private/PythonOnlineDocsCommandlet.cpp
+++ b/Engine/Plugins/Experimental/PythonScriptPlugin/Source/PythonScriptPlugin/Private/PythonOnlineDocsCommandlet.cpp
@@ -34,11 +34,16 @@ int32 UPythonOnlineDocsCommandlet::Main(const FString& Params)
 		DocGenFlags |= EPyOnlineDocsFilterFlags::NAME;	\
 	}
 	APPLY_FILTER_FLAG(IncludeEngine)
-	APPLY_FILTER_FLAG(IncludeEnterprise)
 	APPLY_FILTER_FLAG(IncludeInternal)
 	APPLY_FILTER_FLAG(IncludeProject)
 #undef APPLY_FILTER_FLAG
 
+	if (Switches.Contains(TEXT("IncludeEnterprise")))
+	{
+		UE_LOG(LogPythonOnlineDocsCommandlet, Warning, TEXT("\n\nIncludeEnterprise is deprecated, as Enterprise is now considered part of Engine. Please update callsites to use IncludeEngine."));
+		DocGenFlags |= EPyOnlineDocsFilterFlags::IncludeEngine;
+	}
+
 	// If we weren't given any filter flags, include everything
 	if (DocGenFlags == EPyOnlineDocsFilterFlags::IncludeNone)
 	{
