Python Callback attached to Dynamic Delegate

Calling a python delegate currently crashes in `FFrame::StepExplicitProperty` because the `FFrame::OutParms` can be null even though the property is flagged as `CPF_OutParm`.
The `OutParms` can be null here because `UObject::ProcessEvent` does not create these `OutParms` if the UFunction is not flagged with `FUNC_HasOutParms`.
Looking at the UHT code, this flag is currently only set for functions in a class, and with the delegates being defined outside of a class, this flag is never set.
This leads to the delegates not having the flag `FUNC_HasOutParms`, so the `FFrame::OutParms` is still null when reaching `FFrame::StepExplicitProperty`, but as ref and const-ref parameters would still be tagged with `CPF_OutParm`, a crash would follow.

This CL ensures that the flag `FUNC_HasOutParms` is set for all `UFunctions`, which would now correctly create the `FFrame::OutParms` and stop the crash.
In the case of ref parameters though, the values are currently set to a temporary LocalStruct in `InvokePythonCallableFromUnrealFunctionThunk` which is just discarded at the end of the function, without being given/copied back to the `FFrame`. Instead of creating this temporary object, this CL uses the `FFrame::OutParms` directly, which we have ensured is created in `UObject::ProcessEvent`.

Test code that was originally crashing and that now runs:
```cpp
// PythonTestLibrary.h
#pragma once

#include "CoreMinimal.h"
#include "PythonTestLibrary.generated.h"

USTRUCT(BlueprintType)
struct FPythonTestStruct
{
GENERATED_BODY()

UPROPERTY(BlueprintReadWrite, EditAnywhere, Category = "PythonTestLibrary")
int Int = 0;
UPROPERTY(BlueprintReadWrite, EditAnywhere, Category = "PythonTestLibrary")
bool bIsEven = false;
};

DECLARE_DYNAMIC_DELEGATE_OneParam(FOnPythonTestLibrary_ConstRef, const FPythonTestStruct&, Struct);
DECLARE_DYNAMIC_DELEGATE_RetVal_OneParam(const FPythonTestStruct&, FOnPythonTestLibrary_ConstRef_Return, const FPythonTestStruct&, Struct);
DECLARE_DYNAMIC_DELEGATE_RetVal_OneParam(FPythonTestStruct, FOnPythonTestLibrary_Ref_Return, FPythonTestStruct&, Struct);
```

```python
import unreal

RESET_VALUE = 2
FUNCTION_NEW_RETURN_VALUE = 333

def same_values(a, b):
return a.int == b.int and a.is_even == b.is_even
def reset(val = RESET_VALUE): # resets the struct with the wrong value of is_even
struct = unreal.PythonTestStruct()
struct.int = val
struct.is_even = False
return struct
def SetIsEvenAndReturn(struct):
struct.is_even = (struct.int % 2 == 0)
return struct
def CreateNewStructAndReturnTuple():
struct = reset(RESET_VALUE)
struct.is_even = (struct.int % 2 == 0)
new = reset(FUNCTION_NEW_RETURN_VALUE)
return (new, struct)

### Const-ref return version
### DECLARE_DYNAMIC_DELEGATE_RetVal_OneParam(const FPythonTestStruct&, FOnPythonTestLibrary_ConstRef_Return, const FPythonTestStruct&, Struct);
unreal.log_warning(" === unreal.OnPythonTestLibrary_ConstRef_Return")
struct = reset(RESET_VALUE)
print(" - input:", struct)
delegate = unreal.OnPythonTestLibrary_ConstRef_Return()
delegate.bind_callable(SetIsEvenAndReturn)
ret = delegate.execute(struct)
print("original:", struct)
print("return  :", ret)
assert same_values(struct, unreal.PythonTestStruct(RESET_VALUE, False)) # untouched original
assert same_values(ret, unreal.PythonTestStruct(RESET_VALUE, True)) # the function should have set is_even to true

### Ref return version
### DECLARE_DYNAMIC_DELEGATE_RetVal_OneParam(FPythonTestStruct, FOnPythonTestLibrary_Ref_Return, FPythonTestStruct&, Struct);
unreal.log_warning(" === unreal.OnPythonTestLibrary_Ref_Return")
delegate = unreal.OnPythonTestLibrary_Ref_Return()
delegate.bind_callable(CreateNewStructAndReturnTuple) # return 2 values, one for the return and one for the ref param
ret, ref = delegate.execute()
print("return  :", ret)
print("refparam:", ref)
assert same_values(ret, unreal.PythonTestStruct(FUNCTION_NEW_RETURN_VALUE, False))
assert same_values(ref, unreal.PythonTestStruct(RESET_VALUE, True))  # the function should have set is_even to true
```

diff --git a/Engine/Plugins/Experimental/PythonScriptPlugin/Source/PythonScriptPlugin/Private/PyGenUtil.cpp b/Engine/Plugins/Experimental/PythonScriptPlugin/Source/PythonScriptPlugin/Private/PyGenUtil.cpp
--- a/Engine/Plugins/Experimental/PythonScriptPlugin/Source/PythonScriptPlugin/Private/PyGenUtil.cpp
+++ b/Engine/Plugins/Experimental/PythonScriptPlugin/Source/PythonScriptPlugin/Private/PyGenUtil.cpp
@@ -1200,21 +1200,48 @@ PyObject* PackReturnValues(const void* InBaseParamsAddr, const TArray<FGenerated
 	}
 }
 
-bool UnpackReturnValues(PyObject* InRetVals, const FOutParmRec* InOutputParms, const TCHAR* InErrorCtxt, const TCHAR* InCallingCtxt)
+bool UnpackReturnValues(PyObject* InRetVals, const TOptional<FOutParmRec>& ReturnParam, const FOutParmRec* InOutputParms, const TCHAR* InErrorCtxt, const TCHAR* InCallingCtxt)
 {
 	if (!InOutputParms)
 	{
+		// It is possible for InOutputParms to be null but for the function to still have a return parameter. This is because the UHT flag
+		// EFunctionFlags.HasOutParms is only set on functions having ref parameters. If they do not have ref parameters but have a return parameter,
+		// the flag HasOutParms will not be set and the InOutParms will not be created in UObject::ProcessEvent
+		if (ReturnParam.IsSet())
+		{
+			if (!PyConversion::NativizeProperty_Direct(InRetVals, ReturnParam->Property, ReturnParam->PropAddr))
+			{
+				PyUtil::SetPythonError(PyExc_TypeError, InErrorCtxt, *FString::Printf(TEXT("Failed to convert return property '%s' (%s) when calling %s"), *ReturnParam->Property->GetName(), *ReturnParam->Property->GetClass()->GetName(), InCallingCtxt));
+				return false;
+			}
+		}
 		return true;
 	}
 
 	const FOutParmRec* OutParamRec = InOutputParms;
 
+	int32 NumPropertiesToUnpack = 0;
+	for (const FOutParmRec* Param = OutParamRec; Param; Param = Param->NextOutParm)
+	{
+		// Filter out the const ref properties which are currently part of the OutParms
+		if (PyUtil::IsOutputParameter(Param->Property) || Param->Property->HasAnyPropertyFlags(CPF_ReturnParm) ) 
+		{
+			++NumPropertiesToUnpack;
+		}
+	}
+
+	if (NumPropertiesToUnpack == 0) // If we only had const references
+	{
+		return true;
+	}
+	
 	// If we have multiple return values and the main return value is a bool, we expect None (for false) or the (potentially packed) return value without the bool (for true)
-	if (OutParamRec->NextOutParm && OutParamRec->Property->HasAnyPropertyFlags(CPF_ReturnParm) && OutParamRec->Property->IsA<FBoolProperty>())
+	bool bSkipReturnValue = false; // If true, we do not read the return value as it was not returned. This variable being true also implies that ReturnParam is set
+	if (NumPropertiesToUnpack > 1 && ReturnParam.IsSet() && ReturnParam->Property->IsA<FBoolProperty>())
 	{
-		const FBoolProperty* BoolReturn = CastFieldChecked<const FBoolProperty>(OutParamRec->Property);
+		const FBoolProperty* BoolReturn = CastFieldChecked<const FBoolProperty>(ReturnParam->Property);
 		const bool bReturnValue = InRetVals != Py_None;
-		BoolReturn->SetPropertyValue(OutParamRec->PropAddr, bReturnValue);
+		BoolReturn->SetPropertyValue(ReturnParam->PropAddr, bReturnValue);
 
 		if (!bReturnValue)
 		{
@@ -1222,13 +1249,20 @@ bool UnpackReturnValues(PyObject* InRetVals, const FOutParmRec* InOutputParms, c
 			return true;
 		}
 
-		OutParamRec = OutParamRec->NextOutParm; // Start unpacking at the 1st out value
-		check(OutParamRec);
+		bSkipReturnValue = true;
+		NumPropertiesToUnpack -= 1; // make sure we do not try to unpack a tuple if only one value was returned
 	}
 
 	// Do we need to expect a packed tuple, or just a single value?
-	if (OutParamRec->NextOutParm == nullptr)
+	if (NumPropertiesToUnpack == 1)
 	{
+		// We ensure we are setting the right parameter, which is either the first ref parameter, or the return parameter unless we skip it
+		while (! (PyUtil::IsOutputParameter(OutParamRec->Property) || (OutParamRec->Property->HasAnyPropertyFlags(CPF_ReturnParm) && !bSkipReturnValue)) )
+		{
+			OutParamRec = OutParamRec->NextOutParm;
+			check(OutParamRec);
+		}
+
 		if (!PyConversion::NativizeProperty_Direct(InRetVals, OutParamRec->Property, OutParamRec->PropAddr))
 		{
 			PyUtil::SetPythonError(PyExc_TypeError, InErrorCtxt, *FString::Printf(TEXT("Failed to convert return property '%s' (%s) when calling %s"), *OutParamRec->Property->GetName(), *OutParamRec->Property->GetClass()->GetName(), InCallingCtxt));
@@ -1243,13 +1277,6 @@ bool UnpackReturnValues(PyObject* InRetVals, const FOutParmRec* InOutputParms, c
 			return false;
 		}
 
-		check(OutParamRec->NextOutParm);
-		int32 NumPropertiesToUnpack = 2; // We can start from 2 since we know we already have at least 2 output parameters to get to this code
-		for (const FOutParmRec* Tmp = OutParamRec->NextOutParm->NextOutParm; Tmp; Tmp = Tmp->NextOutParm)
-		{
-			++NumPropertiesToUnpack;
-		}
-
 		const int32 RetTupleSize = PyTuple_Size(InRetVals);
 		if (RetTupleSize != NumPropertiesToUnpack)
 		{
@@ -1257,14 +1284,21 @@ bool UnpackReturnValues(PyObject* InRetVals, const FOutParmRec* InOutputParms, c
 			return false;
 		}
 
-		int32 RetTupleIndex = 0;
+		// The return param is always the first value of the returned tuple so for other out parameters,
+		// start at index 1 if there is a return param and we not need to skip it
+		int32 OutTupleIndex = (ReturnParam.IsSet() && !bSkipReturnValue) ? 1 : 0;
 		do
 		{
-			PyObject* RetVal = PyTuple_GetItem(InRetVals, RetTupleIndex++);
-			if (!PyConversion::NativizeProperty_Direct(RetVal, OutParamRec->Property, OutParamRec->PropAddr))
+			// We ensure we are setting a valid parameter, which is either a ref parameter, or the return parameter unless we skip it
+			if (PyUtil::IsOutputParameter(OutParamRec->Property) || (OutParamRec->Property->HasAnyPropertyFlags(CPF_ReturnParm) && !bSkipReturnValue))
 			{
-				PyUtil::SetPythonError(PyExc_TypeError, InErrorCtxt, *FString::Printf(TEXT("Failed to convert return property '%s' (%s) when calling %s"), *OutParamRec->Property->GetName(), *OutParamRec->Property->GetClass()->GetName(), InCallingCtxt));
-				return false;
+				const int32 TupleIndex = OutParamRec->Property->HasAnyPropertyFlags(CPF_ReturnParm) ? 0 : OutTupleIndex++;
+				PyObject* RetVal = PyTuple_GetItem(InRetVals, TupleIndex);
+				if (!PyConversion::NativizeProperty_Direct(RetVal, OutParamRec->Property, OutParamRec->PropAddr))
+				{
+					PyUtil::SetPythonError(PyExc_TypeError, InErrorCtxt, *FString::Printf(TEXT("Failed to convert return property '%s' (%s) when calling %s"), *OutParamRec->Property->GetName(), *OutParamRec->Property->GetClass()->GetName(), InCallingCtxt));
+					return false;
+				}
 			}
 			OutParamRec = OutParamRec->NextOutParm;
 		}
@@ -1276,52 +1310,41 @@ bool UnpackReturnValues(PyObject* InRetVals, const FOutParmRec* InOutputParms, c
 
 bool InvokePythonCallableFromUnrealFunctionThunk(FPyObjectPtr InSelf, PyObject* InCallable, const UFunction* InFunc, UObject* Context, FFrame& Stack, RESULT_DECL)
 {
-	// Allocate memory to store our local argument data
-	void* LocalStruct = FMemory_Alloca(FMath::Max<int32>(1, InFunc->GetStructureSize()));
-	InFunc->InitializeStruct(LocalStruct);
-	ON_SCOPE_EXIT
-	{
-		InFunc->DestroyStruct(LocalStruct);
-	};
-
 	// Stores information about inputs and outputs
-	FOutParmRec* OutParms = nullptr;
 	TArray<FPyObjectPtr, TInlineAllocator<4>> PyParams;
 
-	// Add any return property to the output params chain
-	if (FProperty* ReturnProp = InFunc->GetReturnProperty())
-	{
-		FOutParmRec* Out = (FOutParmRec*)FMemory_Alloca(sizeof(FOutParmRec));
-		Out->Property = ReturnProp;
-		Out->PropAddr = (uint8*)RESULT_PARAM;
-
-		// Link it to the head of the list, as UnpackReturnValues expects the return value to be first in the list
-		Out->NextOutParm = OutParms;
-		OutParms = Out;
-	}
+	// Stores the potential return parameter which might not always be set in the Stack.OutParms
+	TOptional<FOutParmRec> ReturnParam{};
 
 	// Get the value of the input params for the Python args, and cache the addresses that return and output data should be unpacked to
 	bool bProcessedInputs = true;
 	{
 		int32 ArgIndex = 0;
-		FOutParmRec** LastOut = &OutParms;
+		FOutParmRec** LastOut = &Stack.OutParms; // Stack.OutParms can be null
 
 		// We iterate the fields directly here as we need to process input and output properties in the 
 		// correct stack order, as we're potentially popping data off the bytecode stack
 		for (TFieldIterator<FProperty> ParamIt(InFunc); ParamIt; ++ParamIt)
 		{
 			FProperty* Param = *ParamIt;
+			// The code below assumes that the Out Properties fields and the Stack.OutParms are in the same order
+			check(!Param->HasAnyPropertyFlags(CPF_OutParm) || !(*LastOut) || (*LastOut)->Property == Param);
 
-			// Skip the return value; it never has data on the bytecode stack and was added to the output params chain before this loop
+			// Skip the return value; it never has data on the bytecode stack
 			if (Param->HasAnyPropertyFlags(CPF_ReturnParm))
 			{
+				ReturnParam = FOutParmRec{Param, (uint8*)RESULT_PARAM, nullptr};
+				if (*LastOut) // Even if we have a return parameter Stack.OutParms might be null
+				{
+					LastOut = &(*LastOut)->NextOutParm;
+				}
 				continue;
 			}
 
 			// Step the property data to populate the local value
 			Stack.MostRecentPropertyAddress = nullptr;
 			Stack.MostRecentPropertyContainer = nullptr;
-			void* LocalValue = Param->ContainerPtrToValuePtr<void>(LocalStruct);
+			void* LocalValue = Param->ContainerPtrToValuePtr<void>(Stack.Locals);
 			Stack.StepCompiledIn(LocalValue, Param->GetClass());
 
 			// Output parameters (even const ones) need to read their data from the property address (if available) rather than the local struct
@@ -1331,25 +1354,14 @@ bool InvokePythonCallableFromUnrealFunctionThunk(FPyObjectPtr InSelf, PyObject*
 				ValueAddress = Stack.MostRecentPropertyAddress;
 			}
 
-			// Add any output parameters to the output params chain
 			if (PyUtil::IsOutputParameter(Param))
 			{
-				CA_SUPPRESS(6263) // using _alloca in a loop
-				FOutParmRec* Out = (FOutParmRec*)FMemory_Alloca(sizeof(FOutParmRec));
-				Out->Property = Param;
-				Out->PropAddr = (uint8*)ValueAddress;
-				Out->NextOutParm = nullptr;
-
-				// Link it to the end of the list
-				if (*LastOut)
-				{
-					(*LastOut)->NextOutParm = Out;
-					LastOut = &(*LastOut)->NextOutParm;
-				}
-				else
-				{
-					*LastOut = Out;
-				}
+				(*LastOut)->PropAddr = (uint8*)ValueAddress;
+				LastOut = &(*LastOut)->NextOutParm;
+			}
+			else if (Param->HasAnyPropertyFlags(CPF_OutParm)) // Const ref parameters will return false for IsOutputParameter but are still part of the OutParms 
+			{
+				LastOut = &(*LastOut)->NextOutParm;
 			}
 
 			// Convert any input parameters for use with Python
@@ -1406,7 +1418,7 @@ bool InvokePythonCallableFromUnrealFunctionThunk(FPyObjectPtr InSelf, PyObject*
 	}
 
 	// Unpack any output values
-	if (!UnpackReturnValues(RetVals, OutParms, *PyUtil::GetErrorContext(InCallable), *FString::Printf(TEXT("function '%s' on '%s'"), *InFunc->GetName(), *P_THIS_OBJECT->GetName())))
+	if (!UnpackReturnValues(RetVals, ReturnParam, Stack.OutParms, *PyUtil::GetErrorContext(InCallable), *FString::Printf(TEXT("function '%s' on '%s'"), *InFunc->GetName(), *P_THIS_OBJECT->GetName())))
 	{
 		return false;
 	}
diff --git a/Engine/Plugins/Experimental/PythonScriptPlugin/Source/PythonScriptPlugin/Private/PyGenUtil.h b/Engine/Plugins/Experimental/PythonScriptPlugin/Source/PythonScriptPlugin/Private/PyGenUtil.h
--- a/Engine/Plugins/Experimental/PythonScriptPlugin/Source/PythonScriptPlugin/Private/PyGenUtil.h
+++ b/Engine/Plugins/Experimental/PythonScriptPlugin/Source/PythonScriptPlugin/Private/PyGenUtil.h
@@ -1078,7 +1078,7 @@ namespace PyGenUtil
 	PyObject* PackReturnValues(const void* InBaseParamsAddr, const TArray<FGeneratedWrappedMethodParameter>& InOutputParams, const TCHAR* InErrorCtxt, const TCHAR* InCallingCtxt);
 
 	/** Given a Python return value, unpack the values into the struct data associated with them */
-	bool UnpackReturnValues(PyObject* InRetVals, const FOutParmRec* InOutputParms, const TCHAR* InErrorCtxt, const TCHAR* InCallingCtxt);
+	bool UnpackReturnValues(PyObject* InRetVals, const TOptional<FOutParmRec>& ReturnParam, const FOutParmRec* InOutputParms, const TCHAR* InErrorCtxt, const TCHAR* InCallingCtxt);
 
 	/** Invoke a Python callable from an Unreal function thunk (ie, DEFINE_FUNCTION), ensuring that the calling convention is correct for both native and script function callers */
 	bool InvokePythonCallableFromUnrealFunctionThunk(FPyObjectPtr InSelf, PyObject* InCallable, const UFunction* InFunc, UObject* Context, FFrame& Stack, RESULT_DECL);
diff --git a/Engine/Source/Programs/Shared/EpicGames.UHT/Types/UhtFunction.cs b/Engine/Source/Programs/Shared/EpicGames.UHT/Types/UhtFunction.cs
--- a/Engine/Source/Programs/Shared/EpicGames.UHT/Types/UhtFunction.cs
+++ b/Engine/Source/Programs/Shared/EpicGames.UHT/Types/UhtFunction.cs
@@ -597,27 +597,27 @@ namespace EpicGames.UHT.Types
 						}
 					}
 
-					// The following code is only performed on functions in a class.
-					if (Outer is UhtClass)
+					foreach (UhtType type in Children)
 					{
-						foreach (UhtType type in Children)
+						if (type is not UhtProperty property)
 						{
-							if (type is UhtProperty property)
+							continue;
+						}
+						
+						if (property.PropertyFlags.HasExactFlags(EPropertyFlags.OutParm | EPropertyFlags.ReturnParm, EPropertyFlags.OutParm))
+						{
+							FunctionFlags |= EFunctionFlags.HasOutParms;
+						}
+						// The following code is only performed on functions in a class.
+						if (property is UhtStructProperty structProperty && Outer is UhtClass)
+						{
+							if (structProperty.ScriptStruct.HasDefaults)
 							{
-								if (property.PropertyFlags.HasExactFlags(EPropertyFlags.OutParm | EPropertyFlags.ReturnParm, EPropertyFlags.OutParm))
-								{
-									FunctionFlags |= EFunctionFlags.HasOutParms;
-								}
-								if (property is UhtStructProperty structProperty)
-								{
-									if (structProperty.ScriptStruct.HasDefaults)
-									{
-										FunctionFlags |= EFunctionFlags.HasDefaults;
-									}
-								}
+								FunctionFlags |= EFunctionFlags.HasDefaults;
 							}
 						}
 					}
+
 					break;
 			}
 			return results;
