Fixed issues generating the Python stub file

* Generated wrapped types are now deterministically sorted into an order that can be parsed (parent types before child types).
* This used to be implicit via the generation itself, but the parallel generation broke this assumption.
* Split native Python modules to export their types before the generated types, and their free functions after.
* This addresses some issues where the type hinting for native module free functions depended on generated types to parse correctly.
* Fixed enums having undefined values when exported with type hinting enabled.
* Fixed enums not respecting ScriptName for enum values when Pythonizing their value for argument defaults.
* Fixed deprecated enum entries not being included in the stub file.
* Fixed PropertyAccessChangeNotifyMode being undefined when parsing _ObjectBase and _StructBase.

diff --git a/Engine/Plugins/Experimental/PythonScriptPlugin/Source/PythonScriptPlugin/Private/PyGenUtil.cpp b/Engine/Plugins/Experimental/PythonScriptPlugin/Source/PythonScriptPlugin/Private/PyGenUtil.cpp
--- a/Engine/Plugins/Experimental/PythonScriptPlugin/Source/PythonScriptPlugin/Private/PyGenUtil.cpp
+++ b/Engine/Plugins/Experimental/PythonScriptPlugin/Source/PythonScriptPlugin/Private/PyGenUtil.cpp
@@ -2338,8 +2338,8 @@ void PythonizeValueImpl(const FProperty* InProp, const void* InPropValue, const
 			{
 				const uint8 EnumVal = ByteProp->GetPropertyValue(PropArrValue);
 				const int32 EnumIndex = ByteProp->Enum->GetIndexByValue(EnumVal);
-				const FString EnumValStr = IsScriptExposedEnumEntry(ByteProp->Enum, EnumIndex) ? ByteProp->Enum->GetNameStringByIndex(EnumIndex) : FString();
-				OutPythonDefaultValue += EnumValStr.IsEmpty() ? TEXT("0") : *FString::Printf(TEXT("%s%s.%s"), UnrealNamespace, *GetEnumPythonName(ByteProp->Enum), *PythonizeName(EnumValStr, EPythonizeNameCase::Upper));
+				const FString EnumValStr = IsScriptExposedEnumEntry(ByteProp->Enum, EnumIndex) ? GetEnumEntryPythonName(ByteProp->Enum, EnumIndex) : FString();
+				OutPythonDefaultValue += EnumValStr.IsEmpty() ? TEXT("0") : *FString::Printf(TEXT("%s%s.%s"), UnrealNamespace, *GetEnumPythonName(ByteProp->Enum), *EnumValStr);
 			}
 			else
 			{
@@ -2351,8 +2351,8 @@ void PythonizeValueImpl(const FProperty* InProp, const void* InPropValue, const
 			FNumericProperty* EnumInternalProp = EnumProp->GetUnderlyingProperty();
 			const int64 EnumVal = EnumInternalProp->GetSignedIntPropertyValue(PropArrValue);
 			const int32 EnumIndex = EnumProp->GetEnum()->GetIndexByValue(EnumVal);
-			const FString EnumValStr = IsScriptExposedEnumEntry(EnumProp->GetEnum(), EnumIndex) ? EnumProp->GetEnum()->GetNameStringByIndex(EnumIndex) : FString();
-			OutPythonDefaultValue += EnumValStr.IsEmpty() ? TEXT("0") : *FString::Printf(TEXT("%s%s.%s"), UnrealNamespace, *GetEnumPythonName(EnumProp->GetEnum()), *PythonizeName(EnumValStr, EPythonizeNameCase::Upper));
+			const FString EnumValStr = IsScriptExposedEnumEntry(EnumProp->GetEnum(), EnumIndex) ? GetEnumEntryPythonName(EnumProp->GetEnum(), EnumIndex) : FString();
+			OutPythonDefaultValue += EnumValStr.IsEmpty() ? TEXT("0") : *FString::Printf(TEXT("%s%s.%s"), UnrealNamespace, *GetEnumPythonName(EnumProp->GetEnum()), *EnumValStr);
 		}
 		else if (const FBoolProperty* BoolProp = CastField<const FBoolProperty>(InProp))
 		{
diff --git a/Engine/Plugins/Experimental/PythonScriptPlugin/Source/PythonScriptPlugin/Private/PyWrapperTypeRegistry.cpp b/Engine/Plugins/Experimental/PythonScriptPlugin/Source/PythonScriptPlugin/Private/PyWrapperTypeRegistry.cpp
--- a/Engine/Plugins/Experimental/PythonScriptPlugin/Source/PythonScriptPlugin/Private/PyWrapperTypeRegistry.cpp
+++ b/Engine/Plugins/Experimental/PythonScriptPlugin/Source/PythonScriptPlugin/Private/PyWrapperTypeRegistry.cpp
@@ -3093,40 +3093,39 @@ void FPyWrapperTypeRegistry::GenerateStubCodeForWrappedTypes(const EPyOnlineDocs
 		}
 	}
 
-	// Process native glue code
-	UE_LOG(LogPython, Display, TEXT("  ...generating Python API: glue code"));
+	TSet<FSoftObjectPath> GeneratedTypesToSkip;
+
+	// Process native glue code (types)
+	UE_LOG(LogPython, Display, TEXT("  ...generating Python API: glue code (types)"));
 	PythonScript.WriteLine(TEXT("##### Glue Code #####"));
 	PythonScript.WriteNewLine();
 
 	for (const PyGenUtil::FNativePythonModule& NativePythonModule : NativePythonModules)
 	{
-		for (PyMethodDef* MethodDef = NativePythonModule.PyModuleMethods; MethodDef && MethodDef->ml_name; ++MethodDef)
+		for (PyTypeObject* PyType : NativePythonModule.PyModuleTypes)
 		{
-			// Try parsing the doc string to extract the function declaration with type hinting.
-			FString MethodDecl = UE::Python::GenerateMethodDecl(MethodDef, TEXT("unreal"));
-			PythonScript.WriteLine(MethodDecl);
-			PythonScript.IncreaseIndent();
-			PythonScript.WriteDocString(UTF8_TO_TCHAR(MethodDef->ml_doc));
-			PythonScript.WriteLine(TEXT("...")); // pass
-			PythonScript.DecreaseIndent();
-			PythonScript.WriteNewLine();
+			GenerateStubCodeForWrappedType(PyType, nullptr, PythonScript, OnlineDocsNativeTypesSection.Get());
 
-			if (OnlineDocsUnrealModule.IsValid())
+			if (PyType == &PyWrapperEnumType)
 			{
-				OnlineDocsUnrealModule->AccumulateFunction(UTF8_TO_TCHAR(MethodDef->ml_name));
+				// Include EPropertyAccessChangeNotifyMode in the native types, as it's needed by _ObjectBase and _StructBase (which are after EnumBase in PyModuleTypes)
+				if (const UEnum* PropertyAccessChangeNotifyModeEnum = StaticEnum<EPropertyAccessChangeNotifyMode>())
+				{
+					FSoftObjectPath PropertyAccessChangeNotifyModeEnumPath = PropertyAccessChangeNotifyModeEnum;
+					if (TSharedPtr<PyGenUtil::FGeneratedWrappedType> GeneratedWrappedType = GeneratedWrappedTypes.FindRef(PropertyAccessChangeNotifyModeEnumPath))
+					{
+						GenerateStubCodeForWrappedType(&GeneratedWrappedType->PyType, GeneratedWrappedType.Get(), PythonScript, OnlineDocsEnumTypesSection.Get());
+						GeneratedTypesToSkip.Add(MoveTemp(PropertyAccessChangeNotifyModeEnumPath));
+					}
+				}
 			}
 		}
-
-		for (PyTypeObject* PyType : NativePythonModule.PyModuleTypes)
-		{
-			GenerateStubCodeForWrappedType(PyType, nullptr, PythonScript, OnlineDocsNativeTypesSection.Get());
-		}
 	}
 
 	// Process generated glue code
 	// Also excludes types that don't pass the filters specified in InDocGenFlags using the information about
 	// which module it came from and where that module exists on disk.
-	auto ProcessWrappedDataArray = [this, &PythonScript, InDocGenFlags](const auto& WrappedData, const TSharedPtr<FPyOnlineDocsSection>& OnlineDocsSection)
+	auto ProcessWrappedDataArray = [this, &PythonScript, &GeneratedTypesToSkip, InDocGenFlags](const auto& WrappedData, const TSharedPtr<FPyOnlineDocsSection>& OnlineDocsSection)
 	{
 		if (OnlineDocsSection.IsValid())
 		{
@@ -3141,8 +3140,24 @@ void FPyWrapperTypeRegistry::GenerateStubCodeForWrappedTypes(const EPyOnlineDocs
 			ProjectTopDir = FPaths::GetCleanFilename(FPaths::ProjectDir());
 		}
 
-		WrappedData.ForEach([this, &PythonScript, InDocGenFlags, &ProjectTopDir, &OnlineDocsSection](const auto& WrappedDataPair)
+		struct FSortableGeneratedWrappedType
 		{
+			FString TypePathStr;
+			TSharedPtr<PyGenUtil::FGeneratedWrappedType> GeneratedWrappedType;
+			PyTypeObject* PyTypeObject = nullptr;
+			int32 InheritanceDepth = 0;
+		};
+
+		TArray<FSortableGeneratedWrappedType> SortableWrappedDataArray;
+		SortableWrappedDataArray.Reserve(WrappedData.Num());
+
+		WrappedData.ForEach([this, InDocGenFlags, &ProjectTopDir, &OnlineDocsSection, &GeneratedTypesToSkip, &SortableWrappedDataArray](const auto& WrappedDataPair)
+		{
+			if (GeneratedTypesToSkip.Contains(WrappedDataPair.Key))
+			{
+				return;
+			}
+
 			TSharedPtr<PyGenUtil::FGeneratedWrappedType> GeneratedWrappedType = GeneratedWrappedTypes.FindRef(WrappedDataPair.Key);
 
 			if ((InDocGenFlags != EPyOnlineDocsFilterFlags::IncludeAll) && GeneratedWrappedType.IsValid() && OnlineDocsSection.IsValid())
@@ -3199,8 +3214,35 @@ void FPyWrapperTypeRegistry::GenerateStubCodeForWrappedTypes(const EPyOnlineDocs
 				// else if cannot determine origin then include
 			}
 
-			GenerateStubCodeForWrappedType(WrappedDataPair.Value, GeneratedWrappedType.Get(), PythonScript, OnlineDocsSection.Get());
+			int32 InheritanceDepth = 0;
+			for (PyTypeObject* PyType = WrappedDataPair.Value; PyType; PyType = PyType->tp_base)
+			{
+				++InheritanceDepth;
+			}
+
+			SortableWrappedDataArray.Add({ WrappedDataPair.Key.ToString(), MoveTemp(GeneratedWrappedType), WrappedDataPair.Value, InheritanceDepth });
+		});
+
+		SortableWrappedDataArray.Sort([](const FSortableGeneratedWrappedType& One, const FSortableGeneratedWrappedType& Two)
+		{
+			// Primary sort by inheritance depth, so that parent classes appear before child classes
+			if (One.InheritanceDepth < Two.InheritanceDepth)
+			{
+				return true;
+			}
+			if (One.InheritanceDepth > Two.InheritanceDepth)
+			{
+				return false;
+			}
+
+			// Secondary sort by path name, just for determinism
+			return One.TypePathStr < Two.TypePathStr;
 		});
+
+		for (const FSortableGeneratedWrappedType& SortableWrappedData : SortableWrappedDataArray)
+		{
+			GenerateStubCodeForWrappedType(SortableWrappedData.PyTypeObject, SortableWrappedData.GeneratedWrappedType.Get(), PythonScript, OnlineDocsSection.Get());
+		}
 	};
 
 	ProcessWrappedDataArray(PythonWrappedEnums, OnlineDocsEnumTypesSection);
@@ -3208,6 +3250,31 @@ void FPyWrapperTypeRegistry::GenerateStubCodeForWrappedTypes(const EPyOnlineDocs
 	ProcessWrappedDataArray(PythonWrappedStructs, OnlineDocsStructTypesSection);
 	ProcessWrappedDataArray(PythonWrappedClasses, OnlineDocsClassTypesSection);
 
+	// Process native glue code (methods)
+	UE_LOG(LogPython, Display, TEXT("  ...generating Python API: glue code (methods)"));
+	PythonScript.WriteLine(TEXT("##### Glue Code #####"));
+	PythonScript.WriteNewLine();
+
+	for (const PyGenUtil::FNativePythonModule& NativePythonModule : NativePythonModules)
+	{
+		for (PyMethodDef* MethodDef = NativePythonModule.PyModuleMethods; MethodDef && MethodDef->ml_name; ++MethodDef)
+		{
+			// Try parsing the doc string to extract the function declaration with type hinting.
+			FString MethodDecl = UE::Python::GenerateMethodDecl(MethodDef, TEXT("unreal"));
+			PythonScript.WriteLine(MethodDecl);
+			PythonScript.IncreaseIndent();
+			PythonScript.WriteDocString(UTF8_TO_TCHAR(MethodDef->ml_doc));
+			PythonScript.WriteLine(TEXT("...")); // pass
+			PythonScript.DecreaseIndent();
+			PythonScript.WriteNewLine();
+
+			if (OnlineDocsUnrealModule.IsValid())
+			{
+				OnlineDocsUnrealModule->AccumulateFunction(UTF8_TO_TCHAR(MethodDef->ml_name));
+			}
+		}
+	}
+
 	// Append any additional Python code now that all the reflected API has been exported
 	UE_LOG(LogPython, Display, TEXT("  ...generating Python API: additional code"));
 	PythonScript.WriteLine(TEXT("##### Additional Code #####"));
@@ -3998,8 +4065,20 @@ void FPyWrapperTypeRegistry::GenerateStubCodeForWrappedType(PyTypeObject* PyType
 
 					// The enum value types reported by the engine using print(type(unreal.SomeEnum.SOME_VALUE)) would be 'class SomeEnum' but _EnumEntry class is unrelated.
 					PyGenUtil::IsTypeHintingEnabled()
-						? ExportConstant(*EntryName, PyTypeName, /*Value*/FString(), *EntryDoc)
+						? ExportConstant(*EntryName, PyTypeName, TEXT("..."), *EntryDoc)
 						: ExportConstant(*EntryName, /*TypeHint*/FString(), FString::Printf(TEXT("_EnumEntry(\"%s\", \"%s\", %s)"), *PyTypeName, *EntryName, *EntryValue), *EntryDoc);
+
+					for (const PyGenUtil::FUTF8Buffer& EntryDeprecatedAlias : EnumMember.EntryDeprecatedAliases)
+					{
+						const FString DeprecatedEntryName = UTF8_TO_TCHAR(EntryDeprecatedAlias.GetData());
+						const FString DeprecatedEntryDoc = FString::Printf(TEXT("%s (deprecated: Use %s)"), *EntryDoc, *EntryName);
+
+						PyGenUtil::IsTypeHintingEnabled()
+							? ExportConstant(*DeprecatedEntryName, PyTypeName, EntryName, *DeprecatedEntryDoc)
+							: ExportConstant(*DeprecatedEntryName, /*TypeHint*/FString(), EntryName, *DeprecatedEntryDoc);
+					}
+
+					bHasExportedClassData = true;
 				}
 			}
 		}
