PythonScriptPlugin
- Added PyGenUtilAccessor, to allow some Internal access to PyGenUtil functionality.

diff --git a/Engine/Plugins/Experimental/PythonScriptPlugin/Source/PythonScriptPlugin/Internal/PyGenUtilAccessor.h b/Engine/Plugins/Experimental/PythonScriptPlugin/Source/PythonScriptPlugin/Internal/PyGenUtilAccessor.h
new file mode 100644
--- /dev/null
+++ b/Engine/Plugins/Experimental/PythonScriptPlugin/Source/PythonScriptPlugin/Internal/PyGenUtilAccessor.h
@@ -0,0 +1,33 @@
+﻿// Copyright Epic Games, Inc. All Rights Reserved.
+
+
+#pragma once
+
+
+#if WITH_PYTHON
+
+
+//
+// PyGenUtilAccessor
+//
+// Provides some internal access to functionality in PyGenUtil. Not available for non-Engine use. See corresponding functions in PyGenUtil.
+//
+
+
+namespace PyGenUtilAccessor
+{
+	PYTHONSCRIPTPLUGIN_API void ExtractFunctionParamsAsProperties(const UFunction* InFunc, TArray<const FProperty*>& OutInputParams, TArray<const FProperty*>& OutOutputParams);
+	
+	PYTHONSCRIPTPLUGIN_API FString GetClassPythonName(const UClass* InClass);
+	PYTHONSCRIPTPLUGIN_API FString GetFunctionPythonName(const UFunction* InFunc);
+	PYTHONSCRIPTPLUGIN_API FString GetPropertyPythonName(const FProperty* InProp);
+	PYTHONSCRIPTPLUGIN_API FString GetPropertyPythonType(const FProperty* InProp);
+	
+	PYTHONSCRIPTPLUGIN_API bool ShouldExportClass(const UClass* InClass);
+	PYTHONSCRIPTPLUGIN_API bool ShouldExportFunction(const UFunction* InFunc);
+	PYTHONSCRIPTPLUGIN_API bool ShouldExportProperty(const FProperty* InProp);
+	PYTHONSCRIPTPLUGIN_API bool ShouldExportEditorOnlyProperty(const FProperty* InProp);
+};
+
+
+#endif	// WITH_PYTHON
diff --git a/Engine/Plugins/Experimental/PythonScriptPlugin/Source/PythonScriptPlugin/Private/PyGenUtilAccessor.cpp b/Engine/Plugins/Experimental/PythonScriptPlugin/Source/PythonScriptPlugin/Private/PyGenUtilAccessor.cpp
new file mode 100644
--- /dev/null
+++ b/Engine/Plugins/Experimental/PythonScriptPlugin/Source/PythonScriptPlugin/Private/PyGenUtilAccessor.cpp
@@ -0,0 +1,81 @@
+﻿// Copyright Epic Games, Inc. All Rights Reserved.
+
+
+#include "PyGenUtilAccessor.h"
+
+
+#include "PyGenUtil.h"
+
+
+#if WITH_PYTHON
+
+
+void PyGenUtilAccessor::ExtractFunctionParamsAsProperties(const UFunction* InFunc, TArray<const FProperty*>& OutInputParams, TArray<const FProperty*>& OutOutputParams)
+{
+	OutInputParams.Empty();
+	OutOutputParams.Empty();
+
+	TArray<PyGenUtil::FGeneratedWrappedMethodParameter> InputParams;
+	TArray<PyGenUtil::FGeneratedWrappedMethodParameter> OutputParams;
+	PyGenUtil::ExtractFunctionParams(InFunc, InputParams, OutputParams);
+
+	for (const PyGenUtil::FGeneratedWrappedMethodParameter& InputParam : InputParams)
+	{
+		OutInputParams.Add(InputParam.ParamProp);
+	}
+	for (const PyGenUtil::FGeneratedWrappedMethodParameter& OutputParam : OutputParams)
+	{
+		OutOutputParams.Add(OutputParam.ParamProp);
+	}
+}
+
+
+FString PyGenUtilAccessor::GetClassPythonName(const UClass* InClass)
+{
+	return PyGenUtil::GetClassPythonName(InClass);	
+}
+
+
+FString PyGenUtilAccessor::GetFunctionPythonName(const UFunction* InFunc)
+{
+	return PyGenUtil::GetFunctionPythonName(InFunc);	
+}
+
+
+FString PyGenUtilAccessor::GetPropertyPythonName(const FProperty* InProp)
+{
+	return PyGenUtil::GetPropertyPythonName(InProp);
+}
+
+
+FString PyGenUtilAccessor::GetPropertyPythonType(const FProperty* InProp)
+{
+	return PyGenUtil::GetPropertyPythonType(InProp);
+}
+
+
+bool PyGenUtilAccessor::ShouldExportClass(const UClass* InClass)
+{
+	return PyGenUtil::ShouldExportClass(InClass);
+}
+
+
+bool PyGenUtilAccessor::ShouldExportFunction(const UFunction* InFunc)
+{
+	return PyGenUtil::ShouldExportFunction(InFunc);
+}
+
+
+bool PyGenUtilAccessor::ShouldExportProperty(const FProperty* InProp)
+{
+	return PyGenUtil::ShouldExportProperty(InProp);
+}
+
+
+bool PyGenUtilAccessor::ShouldExportEditorOnlyProperty(const FProperty* InProp)
+{
+	return PyGenUtil::ShouldExportEditorOnlyProperty(InProp);
+}
+
+
+#endif
