Added missing includes and forward decls

#rnx

diff --git a/Engine/Plugins/Experimental/PythonScriptPlugin/Source/PythonScriptPlugin/Internal/PyGenUtilAccessor.h b/Engine/Plugins/Experimental/PythonScriptPlugin/Source/PythonScriptPlugin/Internal/PyGenUtilAccessor.h
--- a/Engine/Plugins/Experimental/PythonScriptPlugin/Source/PythonScriptPlugin/Internal/PyGenUtilAccessor.h
+++ b/Engine/Plugins/Experimental/PythonScriptPlugin/Source/PythonScriptPlugin/Internal/PyGenUtilAccessor.h
@@ -1,8 +1,13 @@
-﻿// Copyright Epic Games, Inc. All Rights Reserved.
-
+// Copyright Epic Games, Inc. All Rights Reserved.
 
 #pragma once
 
+#include "Containers/ContainersFwd.h"
+
+class FProperty;
+class UClass;
+class UFunction;
+
 
 #if WITH_PYTHON
 
