Fixed a crash where we could try and create a redirector over an existing non-redirector object

#rnx

diff --git a/Engine/Plugins/Experimental/PythonScriptPlugin/Source/PythonScriptPlugin/Private/PyCore.cpp b/Engine/Plugins/Experimental/PythonScriptPlugin/Source/PythonScriptPlugin/Private/PyCore.cpp
--- a/Engine/Plugins/Experimental/PythonScriptPlugin/Source/PythonScriptPlugin/Private/PyCore.cpp
+++ b/Engine/Plugins/Experimental/PythonScriptPlugin/Source/PythonScriptPlugin/Private/PyCore.cpp
@@ -96,7 +96,8 @@ UObject* GetPythonTypeContainer()
 UObjectRedirector* CreatePythonTypeLegacyRedirector(const FString& ShortName, const FTopLevelAssetPath& GeneratedPathName)
 {
 	UObject* RedirectorOuter = GetPythonTypeContainer();
-	UObjectRedirector* Redirector = FindObject<UObjectRedirector>(RedirectorOuter, *ShortName);
+	UObject* ExistingObject = FindObject<UObject>(RedirectorOuter, *ShortName);
+	UObjectRedirector* Redirector = Cast<UObjectRedirector>(ExistingObject);
 
 	// If the generated path name matches the redirector path name, then delete any existing redirector as we're about to create an object with that name
 	if (FTopLevelAssetPath(RedirectorOuter->GetFName(), *ShortName) == GeneratedPathName)
@@ -112,7 +113,7 @@ UObjectRedirector* CreatePythonTypeLegacyRedirector(const FString& ShortName, co
 	
 	// Otherwise, create a redirector if one doesn't already exist
 	// The caller is responsible for setting DestinationObject
-	if (!Redirector)
+	if (!Redirector && !ExistingObject)
 	{
 		Redirector = NewObject<UObjectRedirector>(RedirectorOuter, *ShortName, RF_Public | RF_Transient);
 	}
