Fixed crash in ExecutePythonScript when Python was disabled at runtime

This only handled Python being disabled at compile time (!WITH_PYTHON) and would crash if Python support was compiled-in but not initialized.

This change moves the parameter parsing out of ExecuteCustomPythonScriptImpl so that we can simply parse all the parameters (to satisfy the bytecode requirements), and then just early return if Python is disabled or not yet initialized. This also makes ExecuteCustomPythonScriptImpl easier to read/maintain, as it doesn't have sections of WITH_PYTHON scattered within it.

diff --git a/Engine/Plugins/Experimental/PythonScriptPlugin/Source/PythonScriptPlugin/Private/PythonScriptLibrary.cpp b/Engine/Plugins/Experimental/PythonScriptPlugin/Source/PythonScriptPlugin/Private/PythonScriptLibrary.cpp
--- a/Engine/Plugins/Experimental/PythonScriptPlugin/Source/PythonScriptPlugin/Private/PythonScriptLibrary.cpp
+++ b/Engine/Plugins/Experimental/PythonScriptPlugin/Source/PythonScriptPlugin/Private/PythonScriptLibrary.cpp
@@ -76,88 +76,95 @@ bool UPythonScriptLibrary::ExecutePythonScript(const FString& PythonScript, cons
 
 DEFINE_FUNCTION(UPythonScriptLibrary::execExecutePythonScript)
 {
+	struct FParamInfo
+	{
+		FProperty* Property = nullptr;
+		uint8* PropAddr = nullptr;
+		const TCHAR* ParamName = nullptr;
+	};
+	using FParamInfoArray = TArray<FParamInfo, TInlineAllocator<16>>;
+
+	// Read the function params
+	// We must always do this, regardless of whether Python is available/enabled, so that we step the bytecode correctly
+	P_GET_PROPERTY_REF(FStrProperty, PythonScript);
+	P_GET_TARRAY_REF(FString, PythonInputs);
+	P_GET_TARRAY_REF(FString, PythonOutputs);
+	// Read the input values
+	FParamInfoArray InputParams;
+	for (const FString& PythonInput : PythonInputs)
+	{
+		// Note: Medium term the Blueprint interpreter will change to provide us with a list of properties and 
+		// instance pointers, rather than forcing us to jump in and out of its interpreter loop (via StepCompiledIn)
+		Stack.StepCompiledIn<FProperty>(nullptr);
+			
+		FParamInfo& InputParam = InputParams.AddDefaulted_GetRef();
+		InputParam.Property = Stack.MostRecentProperty;
+		InputParam.PropAddr = Stack.MostRecentPropertyAddress;
+		InputParam.ParamName = *PythonInput;
+	}
+	// Read the output values
+	FParamInfoArray OutputParams;
+	for (const FString& PythonOutput : PythonOutputs)
+	{
+		// Note: Medium term the Blueprint interpreter will change to provide us with a list of properties and 
+		// instance pointers, rather than forcing us to jump in and out of its interpreter loop (via StepCompiledIn)
+		Stack.StepCompiledIn<FProperty>(nullptr);
+		check(Stack.MostRecentProperty && Stack.MostRecentPropertyAddress);
+
+		FParamInfo& OutputParam = OutputParams.AddDefaulted_GetRef();
+		OutputParam.Property = Stack.MostRecentProperty;
+		OutputParam.PropAddr = Stack.MostRecentPropertyAddress;
+		OutputParam.ParamName = *PythonOutput;
+	}
+	P_FINISH;
+
+	// Note: Is is safe to return now that P_FINISH has been called
+
+	if (!FPythonScriptPlugin::Get()->IsPythonAvailable())
+	{
+		*(bool*)RESULT_PARAM = false;
+		return;
+	}
+
+	if (!FPythonScriptPlugin::Get()->IsPythonInitialized())
+	{
+		UE_LOG(LogPython, Warning, TEXT("Attempt to execute python command before PythonScriptPlugin is initialized. Ensure your call is after OnPythonInitialized."));
+		*(bool*)RESULT_PARAM = false;
+		return;
+	}
+
+#if WITH_PYTHON
 	auto ExecuteCustomPythonScriptImpl = [&]() -> bool
 	{
 		const FString FunctionErrorName = Stack.Node->GetName();
 		const FString FunctionErrorCtxt = Stack.Node->GetOutermost()->GetName();
 
-#if WITH_PYTHON
 		// Local Python context used when executing this script
 		// Has the inputs written into it prior to execution, and the outputs read from it after execution
 		FPyObjectPtr PyTempGlobalDict = FPyObjectPtr::StealReference(PyDict_Copy(FPythonScriptPlugin::Get()->GetDefaultGlobalDict()));
 		FPyObjectPtr PyTempLocalDict = PyTempGlobalDict;
-#endif	// WITH_PYTHON
-
-		// Read the standard function arguments
-		P_GET_PROPERTY_REF(FStrProperty, PythonScript);
-		P_GET_TARRAY_REF(FString, PythonInputs);
-		P_GET_TARRAY_REF(FString, PythonOutputs);
 
-		// Read the input values and write them to the Python context
-		bool bHasValidInputValues = true;
+		// Write the input values to the Python context
+		for (const FParamInfo& InputParam : InputParams)
 		{
-			for (const FString& PythonInput : PythonInputs)
+			if (!InputParam.Property)
 			{
-				// Note: Medium term the Blueprint interpreter will change to provide us with a list of properties and 
-				// instance pointers, rather than forcing us to jump in and out of its interpreter loop (via StepCompiledIn)
-				Stack.StepCompiledIn<FProperty>(nullptr);
-				if(!Stack.MostRecentProperty)
-				{
-#if WITH_PYTHON
-					// an exception was generated when evaluating the input, provide the user script with None:
-					PyDict_SetItemString(PyTempLocalDict, TCHAR_TO_UTF8(*PythonInput), Py_None);
-#endif	// WITH_PYTHON
-					continue;
-				}
-				else
-				{
-					check(Stack.MostRecentPropertyAddress);
-				}
-
-#if WITH_PYTHON
-				FPyObjectPtr PyInput;
-				if (PyConversion::PythonizeProperty_Direct(Stack.MostRecentProperty, Stack.MostRecentPropertyAddress, PyInput.Get()))
-				{
-					PyDict_SetItemString(PyTempLocalDict, TCHAR_TO_UTF8(*PythonInput), PyInput);
-				}
-				else
-				{
-					PyUtil::SetPythonError(PyExc_TypeError, *FunctionErrorCtxt, *FString::Printf(TEXT("Failed to convert input property '%s' (%s) to attribute '%s' when calling function '%s' on '%s'"), *Stack.MostRecentProperty->GetName(), *Stack.MostRecentProperty->GetClass()->GetName(), *PythonInput, *FunctionErrorName, *P_THIS_OBJECT->GetName()));
-					bHasValidInputValues = false;
-				}
-#endif	// WITH_PYTHON
+				// an exception was generated when evaluating the input, provide the user script with None:
+				PyDict_SetItemString(PyTempLocalDict, TCHAR_TO_UTF8(InputParam.ParamName), Py_None);
+				continue;
 			}
-		}
-
-		struct FPyOutputParam
-		{
-			FProperty* Property = nullptr;
-			uint8* PropAddr = nullptr;
-			const TCHAR* OutputName = nullptr;
-		};
-
-		// Read the output values and store them to write to later from the Python context
-		TArray<FPyOutputParam, TInlineAllocator<4>> OutParms;
-		for (const FString& PythonOutput : PythonOutputs)
-		{
-			// Note: Medium term the Blueprint interpreter will change to provide us with a list of properties and 
-			// instance pointers, rather than forcing us to jump in and out of its interpreter loop (via StepCompiledIn)
-			Stack.StepCompiledIn<FProperty>(nullptr);
-			check(Stack.MostRecentProperty && Stack.MostRecentPropertyAddress);
-
-			FPyOutputParam& OutParam = OutParms.AddDefaulted_GetRef();
-			OutParam.Property = Stack.MostRecentProperty;
-			OutParam.PropAddr = Stack.MostRecentPropertyAddress;
-			OutParam.OutputName = *PythonOutput;
-		}
-
-		P_FINISH;
+			check(InputParam.PropAddr);
 
-#if	WITH_PYTHON
-		// If we already failed during the parameter processing phase, then just bail now
-		if (!bHasValidInputValues)
-		{
-			return false;
+			FPyObjectPtr PyInput;
+			if (PyConversion::PythonizeProperty_Direct(InputParam.Property, InputParam.PropAddr, PyInput.Get()))
+			{
+				PyDict_SetItemString(PyTempLocalDict, TCHAR_TO_UTF8(InputParam.ParamName), PyInput);
+			}
+			else
+			{
+				PyUtil::SetPythonError(PyExc_TypeError, *FunctionErrorCtxt, *FString::Printf(TEXT("Failed to convert input property '%s' (%s) to attribute '%s' when calling function '%s' on '%s'"), *InputParam.Property->GetName(), *InputParam.Property->GetClass()->GetName(), InputParam.ParamName, *FunctionErrorName, *P_THIS_OBJECT->GetName()));
+				return false;
+			}
 		}
 
 		// Execute the Python command
@@ -166,30 +173,28 @@ DEFINE_FUNCTION(UPythonScriptLibrary::execExecutePythonScript)
 		// Read the output values from the Python context
 		if (PyResult)
 		{
-			for (const FPyOutputParam& OutParam : OutParms)
+			for (const FParamInfo& OutputParam : OutputParams)
 			{
-				PyObject* PyOutput = PyDict_GetItemString(PyTempLocalDict, TCHAR_TO_UTF8(OutParam.OutputName));
+				PyObject* PyOutput = PyDict_GetItemString(PyTempLocalDict, TCHAR_TO_UTF8(OutputParam.ParamName));
 				if (!PyOutput)
 				{
-					PyUtil::SetPythonError(PyExc_TypeError, *FunctionErrorCtxt, *FString::Printf(TEXT("Failed to find attribute '%s' for output property '%s' (%s) when calling function '%s' on '%s'"), OutParam.OutputName, *OutParam.Property->GetName(), *OutParam.Property->GetClass()->GetName() , *FunctionErrorName, *P_THIS_OBJECT->GetName()));
+					PyUtil::SetPythonError(PyExc_TypeError, *FunctionErrorCtxt, *FString::Printf(TEXT("Failed to find attribute '%s' for output property '%s' (%s) when calling function '%s' on '%s'"), OutputParam.ParamName, *OutputParam.Property->GetName(), *OutputParam.Property->GetClass()->GetName() , *FunctionErrorName, *P_THIS_OBJECT->GetName()));
 					return false;
 				}
 
-				if (!PyConversion::NativizeProperty_Direct(PyOutput, OutParam.Property, OutParam.PropAddr))
+				if (!PyConversion::NativizeProperty_Direct(PyOutput, OutputParam.Property, OutputParam.PropAddr))
 				{
-					PyUtil::SetPythonError(PyExc_TypeError, *FunctionErrorCtxt, *FString::Printf(TEXT("Failed to convert output property '%s' (%s) from attribute '%s' when calling function '%s' on '%s'"), *OutParam.Property->GetName(), *OutParam.Property->GetClass()->GetName(), OutParam.OutputName, *FunctionErrorName, *P_THIS_OBJECT->GetName()));
+					PyUtil::SetPythonError(PyExc_TypeError, *FunctionErrorCtxt, *FString::Printf(TEXT("Failed to convert output property '%s' (%s) from attribute '%s' when calling function '%s' on '%s'"), *OutputParam.Property->GetName(), *OutputParam.Property->GetClass()->GetName(), OutputParam.ParamName, *FunctionErrorName, *P_THIS_OBJECT->GetName()));
 					return false;
 				}
 			}
 
 			return true;
 		}
-#endif	// WITH_PYTHON
 
 		return false;
 	};
 
-#if WITH_PYTHON
 	// Execute Python code within this block
 	{
 		FPyScopedGIL GIL;
@@ -206,10 +211,5 @@ DEFINE_FUNCTION(UPythonScriptLibrary::execExecutePythonScript)
 			*(bool*)RESULT_PARAM = false;
 		}
 	}
-#else	// WITH_PYTHON
-	// We still need to call this function to step the bytecode correctly...
-	ExecuteCustomPythonScriptImpl();
-	*(bool*)RESULT_PARAM = false;
 #endif	// WITH_PYTHON
 }
-
