Encoding Fix for PipInstall

Before Python 3.13, the Python library function `site.addsitedir` uses a `TextIOWrapper` with a default `locale` encoding to read the `.pth` files. This encoding is inconsistent depending on the platform while we need the `.pth` to be consistently read as utf-8.
This has now been fixed in the Python library in Python 3.13, but as a workaround for previous Python versions, we temporarily redefine `TextIOWrapper` to force a utf-8 encoding.

diff --git a/Engine/Plugins/Experimental/PythonScriptPlugin/Source/PythonScriptPlugin/Private/PipInstall.cpp b/Engine/Plugins/Experimental/PythonScriptPlugin/Source/PythonScriptPlugin/Private/PipInstall.cpp
--- a/Engine/Plugins/Experimental/PythonScriptPlugin/Source/PythonScriptPlugin/Private/PipInstall.cpp
+++ b/Engine/Plugins/Experimental/PythonScriptPlugin/Source/PythonScriptPlugin/Private/PipInstall.cpp
@@ -699,7 +699,7 @@ void FPipInstall::WriteSitePackagePthFile() const
 
 	// Create .pth file in PipInstall/Lib/site-packages to account for plugins with packaged dependencies
 	const FString PyPluginsSitePackageFile = FPaths::ConvertRelativePathToFull(GetPipSitePackagesPath() / PluginsSitePackageFilename);
-	FFileHelper::SaveStringArrayToFile(PluginSitePackagePaths, *PyPluginsSitePackageFile, FFileHelper::EEncodingOptions::ForceUTF8WithoutBOM);
+	FFileHelper::SaveStringArrayToFile(PluginSitePackagePaths, *PyPluginsSitePackageFile, FFileHelper::EEncodingOptions::ForceUTF8);
 }
 
 
diff --git a/Engine/Plugins/Experimental/PythonScriptPlugin/Source/PythonScriptPlugin/Private/PyUtil.cpp b/Engine/Plugins/Experimental/PythonScriptPlugin/Source/PythonScriptPlugin/Private/PyUtil.cpp
--- a/Engine/Plugins/Experimental/PythonScriptPlugin/Source/PythonScriptPlugin/Private/PyUtil.cpp
+++ b/Engine/Plugins/Experimental/PythonScriptPlugin/Source/PythonScriptPlugin/Private/PyUtil.cpp
@@ -1,6 +1,8 @@
 // Copyright Epic Games, Inc. All Rights Reserved.
 
 #include "PyUtil.h"
+
+#include "IPythonScriptPlugin.h"
 #include "PyCore.h"
 #include "PyConversion.h"
 #include "PyGIL.h"
@@ -1129,7 +1131,9 @@ void AddSitePackagesPath(const FString& InPath)
 	{
 		return;
 	}
-
+	
+#if PY_MAJOR_VERSION >= 3 && PY_MINOR_VERSION >= 13
+	// We can safely call site.addsitedir on Python 3.13 as it is now decoded first with utf-8 https://docs.python.org/3/library/site.html
 	if (FPyObjectPtr PySiteModule = FPyObjectPtr::StealReference(PyImport_ImportModule("site")))
 	{
 		PyObject* PySiteDict = PyModule_GetDict(PySiteModule);
@@ -1139,9 +1143,36 @@ void AddSitePackagesPath(const FString& InPath)
 			if (PyConversion::Pythonize(InPath, PyPath.Get(), PyConversion::ESetErrorState::No))
 			{
 				FPyObjectPtr PyAddSiteDirResult = FPyObjectPtr::StealReference(PyObject_CallFunctionObjArgs(PyAddSiteDirFunc, PyPath.Get(), nullptr));
+				
+				if (!PyAddSiteDirResult)
+				{
+					PyUtil::LogPythonError(); // Ensure the error is logged and cleared
+				}
 			}
 		}
 	}
+#else
+	// Before Python 3.13, we cannot call site.addsitedir directly as it uses a TextIOWrapper with a default "locale" encoding to read .pth files 
+	// (which is inconsistent depending on the platform) while we need the .pth to be consistently read as utf-8.
+	// As a workaround, we temporarily redefine TextIOWrapper to force a utf-8 encoding
+	FString PyCommandStr = FString::Printf(TEXT(
+	"import io\n"
+	"import site\n"
+	
+	"old_TextIOWrapper = io.TextIOWrapper\n" // Save the existing io.TextIOWrapper implementation
+	"def TextIOWrapper_utf8(buffer, encoding = None):\n"
+	"    return old_TextIOWrapper(buffer, encoding='utf-8')\n" // Create a new io.TextIOWrapper forcing the encoding to be utf-8
+	
+	"io.TextIOWrapper = TextIOWrapper_utf8\n" // Redefine io.TextIOWrapper
+	
+	"try:\n"
+	"    site.addsitedir(r\"%s\")\n" // sets the sitedir
+	"finally:\n"
+	"    io.TextIOWrapper = old_TextIOWrapper\n" // Reset io.TextIOWrapper to the original implementation
+	), *InPath);
+	
+	ensure(IPythonScriptPlugin::Get()->ExecPythonCommand(*PyCommandStr));
+#endif
 }
 
 void AddSystemPath(const FString& InPath)
