PipInstall: Made pip install deferred flow more similar to script plugin, also added helpers for standard python init/pip install deferred flow.

diff --git a/Engine/Plugins/Experimental/PythonScriptPlugin/Source/PythonScriptPlugin/Private/PipInstall.cpp b/Engine/Plugins/Experimental/PythonScriptPlugin/Source/PythonScriptPlugin/Private/PipInstall.cpp
--- a/Engine/Plugins/Experimental/PythonScriptPlugin/Source/PythonScriptPlugin/Private/PipInstall.cpp
+++ b/Engine/Plugins/Experimental/PythonScriptPlugin/Source/PythonScriptPlugin/Private/PipInstall.cpp
@@ -8,6 +8,7 @@
 #include "PythonScriptPluginSettings.h"
 #include "Dom/JsonObject.h"
 #include "HAL/PlatformFileManager.h"
+#include "Async/Async.h"
 #include "Async/Mutex.h"
 #include "Async/ParallelFor.h"
 #include "Misc/CommandLine.h"
@@ -82,7 +83,7 @@ bool FPipInstall::InitPipInstall()
 #endif // WITH_PYTHON
 }
 
-bool FPipInstall::LaunchPipInstall(bool RunAsync, TSharedPtr<ICmdProgressNotifier> CmdProgressNotifier)
+bool FPipInstall::LaunchPipInstall(bool RunAsync, TSharedPtr<ICmdProgressNotifier> CmdProgressNotifier, FSimpleDelegate OnCompleted)
 {
 	TRACE_CPUPROFILER_EVENT_SCOPE(FPipInstall::LaunchPipInstallAsync);
 
@@ -102,7 +103,7 @@ bool FPipInstall::LaunchPipInstall(bool RunAsync, TSharedPtr<ICmdProgressNotifie
 	FFeedbackContext* Context = GWarn;
 
 	// Run install of all python dependencies for enabled plugins
-	return LaunchPipInstallDependencies(RunAsync, CmdProgressNotifier, Context);
+	return LaunchPipInstallImpl(RunAsync, CmdProgressNotifier, Context, MoveTemp(OnCompleted));
 #else
 	return false;
 #endif
@@ -490,9 +491,9 @@ bool FPipInstall::ParsePluginDependencies(const FString& MergedInRequirementsFil
 	return (RunPythonCmd(VenvInterp, Cmd, Context) == 0);
 }
 
-bool FPipInstall::LaunchPipInstallDependencies(bool RunAsync, TSharedPtr<ICmdProgressNotifier> CmdProgressNotifier, FFeedbackContext* Context)
+bool FPipInstall::LaunchPipInstallImpl(bool RunAsync, TSharedPtr<ICmdProgressNotifier> CmdProgressNotifier, FFeedbackContext* Context, FSimpleDelegate OnCompleted)
 {
-	TRACE_CPUPROFILER_EVENT_SCOPE(FPipInstall::LaunchPipInstallDependencies);
+	TRACE_CPUPROFILER_EVENT_SCOPE(FPipInstall::LaunchPipInstallImpl);
 
 	if (IsBackgroundInstalling())
 	{
@@ -530,10 +531,10 @@ bool FPipInstall::LaunchPipInstallDependencies(bool RunAsync, TSharedPtr<ICmdPro
 
 	if (RunAsync)
 	{
-		return RunPipCmdAsync(VenvInterp, Cmd, Context, ProgressParser);
+		return RunPipCmdAsync(VenvInterp, Cmd, Context, ProgressParser, MoveTemp(OnCompleted));
 	}
 
-	return RunPipCmdSync(VenvInterp, Cmd, Context, ProgressParser);
+	return RunPipCmdSync(VenvInterp, Cmd, Context, ProgressParser, MoveTemp(OnCompleted));
 }
 
 FString FPipInstall::GetPipInstallPath() const
@@ -789,30 +790,40 @@ int32 FPipInstall::RunPythonCmd(const FString& PythonInterp, const FString& Cmd,
 	return OutResult;
 }
 
-bool FPipInstall::RunPipCmdAsync(const FString& PythonInterp, const FString& Cmd, FFeedbackContext* Context, TSharedPtr<ICmdProgressParser> CmdParser)
+bool FPipInstall::RunPipCmdAsync(const FString& PythonInterp, const FString& Cmd, FFeedbackContext* Context, TSharedPtr<ICmdProgressParser> CmdParser, FSimpleDelegate OnCompleted)
 {
 	TRACE_CPUPROFILER_EVENT_SCOPE(FPipInstall::RunPythonCmd);
 	
 	UE_LOG(LogPython, Log, TEXT("Running background pip install: python %s"), *Cmd);
 	BackgroundInstallRunnable = MakeShared<FLoggedSubprocessThread>(FPaths::ConvertRelativePathToFull(PythonInterp), Cmd, Context, CmdParser);
-	BackgroundInstallRunnable->OnCompleted().BindLambda([this](int32 ReturnCode)
+	BackgroundInstallRunnable->OnCompleted().BindLambda([this, OnCompleted = MoveTemp(OnCompleted)](int32 ReturnCode)
 	{
 		if (ReturnCode == 0)
 		{
 			CacheDetectInstallDeps();
+
+			if (OnCompleted.IsBound())
+			{
+				AsyncTask(ENamedThreads::GameThread, [OnCompleted]()
+				{
+					OnCompleted.Execute();
+				});
+			}
 		}
 	});
 	
 	return BackgroundInstallRunnable->Launch();
 }
 
-bool FPipInstall::RunPipCmdSync(const FString& PythonInterp, const FString& Cmd, FFeedbackContext* Context, TSharedPtr<ICmdProgressParser> CmdParser)
+bool FPipInstall::RunPipCmdSync(const FString& PythonInterp, const FString& Cmd, FFeedbackContext* Context, TSharedPtr<ICmdProgressParser> CmdParser, FSimpleDelegate OnCompleted)
 {
 	int Res = RunPythonCmd(VenvInterp, Cmd, Context, CmdParser);
 	if (Res == 0)
 	{
 		// TODO: Update regardless of success?
-		return CacheDetectInstallDeps();
+		bool bCached = CacheDetectInstallDeps();
+		OnCompleted.ExecuteIfBound();
+		return bCached;
 	}
 
 	return false;
diff --git a/Engine/Plugins/Experimental/PythonScriptPlugin/Source/PythonScriptPlugin/Private/PipInstall.h b/Engine/Plugins/Experimental/PythonScriptPlugin/Source/PythonScriptPlugin/Private/PipInstall.h
--- a/Engine/Plugins/Experimental/PythonScriptPlugin/Source/PythonScriptPlugin/Private/PipInstall.h
+++ b/Engine/Plugins/Experimental/PythonScriptPlugin/Source/PythonScriptPlugin/Private/PipInstall.h
@@ -25,7 +25,7 @@ public:
 	// IPipInstall Interface Implementation
 	virtual bool InitPipInstall() override;
 	// Run pip to install all missing dependencies for enabled plugins
-	virtual bool LaunchPipInstall(bool RunAsync, TSharedPtr<ICmdProgressNotifier> CmdProgressNotifier = nullptr) override;
+	virtual bool LaunchPipInstall(bool RunAsync, TSharedPtr<ICmdProgressNotifier> CmdProgressNotifier = nullptr, FSimpleDelegate OnCompleted = FSimpleDelegate()) override;
 	// Check if a background install is still running
 	virtual bool IsInstalling() override;
 	
@@ -56,7 +56,7 @@ private:
 	bool SetupPipEnv(FFeedbackContext* Context, bool bForceRebuild = false) const;
 	void RemoveParsedDependencyFiles() const;
 	bool ParsePluginDependencies(const FString& MergedInRequirementsFile, FFeedbackContext* Context) const;
-	bool LaunchPipInstallDependencies(bool RunAsync, TSharedPtr<ICmdProgressNotifier> CmdProgressNotifier, FFeedbackContext* Context);
+	bool LaunchPipInstallImpl(bool RunAsync, TSharedPtr<ICmdProgressNotifier> CmdProgressNotifier, FFeedbackContext* Context, FSimpleDelegate OnCompleted);
 
 	FString GetPipInstallPath() const;
 	FString GetPipSitePackagesPath() const;
@@ -79,8 +79,8 @@ private:
 
 	FString ParseVenvVersion() const;
 	FString SetupPipInstallCmd(const FString& ParsedReqsFile, const TArray<FString>& ExtraUrls) const;
-	bool RunPipCmdAsync(const FString& PythonInterp, const FString& Cmd, FFeedbackContext* Context, TSharedPtr<ICmdProgressParser> CmdParser);
-	bool RunPipCmdSync(const FString& PythonInterp, const FString& Cmd, FFeedbackContext* Context, TSharedPtr<ICmdProgressParser> CmdParser);
+	bool RunPipCmdAsync(const FString& PythonInterp, const FString& Cmd, FFeedbackContext* Context, TSharedPtr<ICmdProgressParser> CmdParser, FSimpleDelegate OnCompleted = FSimpleDelegate());
+	bool RunPipCmdSync(const FString& PythonInterp, const FString& Cmd, FFeedbackContext* Context, TSharedPtr<ICmdProgressParser> CmdParser, FSimpleDelegate OnCompleted = FSimpleDelegate());
 
 	static int32 RunPythonCmd(const FString& PythonInterp, const FString& Cmd, FFeedbackContext* Context, TSharedPtr<ICmdProgressParser> CmdParser = nullptr);
 
diff --git a/Engine/Plugins/Experimental/PythonScriptPlugin/Source/PythonScriptPlugin/Private/PipInstallHelpers.cpp b/Engine/Plugins/Experimental/PythonScriptPlugin/Source/PythonScriptPlugin/Private/PipInstallHelpers.cpp
--- a/Engine/Plugins/Experimental/PythonScriptPlugin/Source/PythonScriptPlugin/Private/PipInstallHelpers.cpp
+++ b/Engine/Plugins/Experimental/PythonScriptPlugin/Source/PythonScriptPlugin/Private/PipInstallHelpers.cpp
@@ -2,6 +2,7 @@
 
 #include "PipInstallHelpers.h"
 #include "IPipInstall.h"
+#include "IPythonScriptPlugin.h"
 #include "PipInstallLauncher.h"
 #include "PipInstallCmdNotifiers.h"
 #include "PyUtil.h"
@@ -244,7 +245,7 @@ int FPipInstallHelper::GetNumPackagesToInstall()
 }
 
 
-EPipInstallDialogResult FPipInstallHelper::ShowPipInstallDialog(bool bAllowBackgroundInstall)
+EPipInstallDialogResult FPipInstallHelper::ShowPipInstallDialog(bool bAllowBackgroundInstall, FSimpleDelegate OnCompleted)
 {
 	check(IsInGameThread());
 	
@@ -254,17 +255,32 @@ EPipInstallDialogResult FPipInstallHelper::ShowPipInstallDialog(bool bAllowBackg
 		UE_LOG(LogPython, Error, TEXT("Unable to initialize Pip Install"));
 		return EPipInstallDialogResult::Error;
 	}
+
+	//if (PipInstall.IsInstalling())
+	//{
+	//	return EPipInstallDialogResult::BackgroundInstall;
+	//}
 	
 	bool bClickedInstall = ShowPackageInstallDialogModal(PipInstall, bAllowBackgroundInstall);
+
+	// Check we actually need to install anything
+	if (PipInstall.GetNumPackagesToInstall() == 0)
+	{
+		return EPipInstallDialogResult::Finished;
+	}
 	
 	// Install must be run immediately (game-thread) if background disallowed
 	if (!bAllowBackgroundInstall)
 	{
-		return (PipInstallLauncher::StartSync(PipInstall) ? EPipInstallDialogResult::Finished : EPipInstallDialogResult::Error); 
+		return (PipInstallLauncher::StartSync(PipInstall, MoveTemp(OnCompleted)) ? EPipInstallDialogResult::Finished : EPipInstallDialogResult::Error); 
 	}
-	else if (bClickedInstall && !PipInstall.IsInstalling())
+	else if (PipInstall.IsInstalling())
 	{
-		return (PipInstallLauncher::StartAsync(PipInstall) ? EPipInstallDialogResult::BackgroundInstall : EPipInstallDialogResult::Error);
+		return EPipInstallDialogResult::BackgroundInstall;
+	}
+	else if (bClickedInstall)
+	{
+		return (PipInstallLauncher::StartAsync(PipInstall, MoveTemp(OnCompleted)) ? EPipInstallDialogResult::BackgroundInstall : EPipInstallDialogResult::Error);
 	}
 	
 	return EPipInstallDialogResult::Canceled;
@@ -284,18 +300,52 @@ bool FPipInstallHelper::LaunchHeadlessPipInstall()
 	return PipInstallLauncher::StartSync(PipInstall); 
 }
 
+void FPythonScriptInitHelper::InitPython(FSimpleDelegate OnInitialized, FSimpleDelegate OnPythonUnavailable)
+{
+	IPythonScriptPlugin::Get()->ForceEnablePythonAtRuntime();
+	IPythonScriptPlugin::Get()->RegisterOnPythonConfigured(FSimpleDelegate::CreateLambda([OnInitialized = MoveTemp(OnInitialized), OnPythonUnavailable = MoveTemp(OnPythonUnavailable)]() mutable
+	{
+		if (IPythonScriptPlugin::Get()->IsPythonAvailable())
+		{
+			IPythonScriptPlugin::Get()->RegisterOnPythonInitialized(MoveTemp(OnInitialized));
+		}
+		else
+		{
+			OnPythonUnavailable.ExecuteIfBound();
+		}
+	}));
+}
+
+void FPythonScriptInitHelper::InitPythonAndPipInstall(FSimpleDelegate OnPipInstalled, FSimpleDelegate OnPythonUnavailable)
+{
+	InitPython(FSimpleDelegate::CreateLambda([OnPipInstalled = MoveTemp(OnPipInstalled)]()
+	{
+		// Run deferred pip install through UI interface if necessary
+		int NumPackages = FPipInstallHelper::GetNumPackagesToInstall();
+		if (NumPackages > 0)
+		{
+			FPipInstallHelper::ShowPipInstallDialog(true, OnPipInstalled);
+		}
+		else
+		{
+			OnPipInstalled.ExecuteIfBound();
+		}
+	}), MoveTemp(OnPythonUnavailable));
+}
+
+
 namespace PipInstallLauncher
 {
-	bool StartSync(IPipInstall& PipInstall)
+	bool StartSync(IPipInstall& PipInstall, FSimpleDelegate OnCompleted)
 	{
 		TSharedPtr<ICmdProgressNotifier> CmdNotifier = MakeShared<FSlowTaskNotifier>(PipInstall.GetNumPackagesToInstall(), LOCTEXT("PipInstall.FGInstallText", "Installing Python Dependencies..."));
-		return PipInstall.LaunchPipInstall(false, CmdNotifier);
+		return PipInstall.LaunchPipInstall(false, CmdNotifier, MoveTemp(OnCompleted));
 	}
 
-	bool StartAsync(IPipInstall& PipInstall)
+	bool StartAsync(IPipInstall& PipInstall, FSimpleDelegate OnCompleted)
 	{
 		TSharedPtr<ICmdProgressNotifier> CmdProgressNotifier = MakeShared<FAsyncTaskCmdNotifier>(PipInstall.GetNumPackagesToInstall(), LOCTEXT("PipInstall.BGInstallText", "Installing Python Dependencies..."));
-		return PipInstall.LaunchPipInstall(true, CmdProgressNotifier);
+		return PipInstall.LaunchPipInstall(true, CmdProgressNotifier, MoveTemp(OnCompleted));
 	}
 }
 
diff --git a/Engine/Plugins/Experimental/PythonScriptPlugin/Source/PythonScriptPlugin/Private/PipInstallLauncher.h b/Engine/Plugins/Experimental/PythonScriptPlugin/Source/PythonScriptPlugin/Private/PipInstallLauncher.h
--- a/Engine/Plugins/Experimental/PythonScriptPlugin/Source/PythonScriptPlugin/Private/PipInstallLauncher.h
+++ b/Engine/Plugins/Experimental/PythonScriptPlugin/Source/PythonScriptPlugin/Private/PipInstallLauncher.h
@@ -6,6 +6,6 @@
 
 namespace PipInstallLauncher
 {
-	bool StartSync(IPipInstall& PipInstall);
-	bool StartAsync(IPipInstall& PipInstall);
+	bool StartSync(IPipInstall& PipInstall, FSimpleDelegate OnCompleted = FSimpleDelegate());
+	bool StartAsync(IPipInstall& PipInstall, FSimpleDelegate OnCompleted = FSimpleDelegate());
 }
diff --git a/Engine/Plugins/Experimental/PythonScriptPlugin/Source/PythonScriptPlugin/Public/IPipInstall.h b/Engine/Plugins/Experimental/PythonScriptPlugin/Source/PythonScriptPlugin/Public/IPipInstall.h
--- a/Engine/Plugins/Experimental/PythonScriptPlugin/Source/PythonScriptPlugin/Public/IPipInstall.h
+++ b/Engine/Plugins/Experimental/PythonScriptPlugin/Source/PythonScriptPlugin/Public/IPipInstall.h
@@ -28,7 +28,7 @@ public:
 	// Initialize the internal pip virtual env and check enabled plugins for python dependencies
 	virtual bool InitPipInstall() = 0;
 	// Run pip to install all missing python dependencies for enabled plugins
-	virtual bool LaunchPipInstall(bool RunAsync, TSharedPtr<ICmdProgressNotifier> CmdProgressNotifier = nullptr) = 0;
+	virtual bool LaunchPipInstall(bool RunAsync, TSharedPtr<ICmdProgressNotifier> CmdProgressNotifier = nullptr, FSimpleDelegate OnCompleted = FSimpleDelegate()) = 0;
 	// Check if a background install is running
 	virtual bool IsInstalling() = 0;
 	
diff --git a/Engine/Plugins/Experimental/PythonScriptPlugin/Source/PythonScriptPlugin/Public/PipInstallHelpers.h b/Engine/Plugins/Experimental/PythonScriptPlugin/Source/PythonScriptPlugin/Public/PipInstallHelpers.h
--- a/Engine/Plugins/Experimental/PythonScriptPlugin/Source/PythonScriptPlugin/Public/PipInstallHelpers.h
+++ b/Engine/Plugins/Experimental/PythonScriptPlugin/Source/PythonScriptPlugin/Public/PipInstallHelpers.h
@@ -25,10 +25,20 @@ public:
 	static UE_API int32 GetNumPackagesToInstall();
 	
 	// Show notification that pip install is required
-	static UE_API EPipInstallDialogResult ShowPipInstallDialog(bool bAllowBackgroundInstall = true);
+	static UE_API EPipInstallDialogResult ShowPipInstallDialog(bool bAllowBackgroundInstall = true, FSimpleDelegate OnCompleted = FSimpleDelegate());
 
 	// Run a headless pip install (for commandlets/build machines)
 	static UE_API bool LaunchHeadlessPipInstall();
 };
 
+/**
+ * Helpers for handling standard delegate pipeline pattern for python/pip install
+ */
+class FPythonScriptInitHelper
+{
+public:
+	static UE_API void InitPython(FSimpleDelegate OnInitialized, FSimpleDelegate OnPythonUnavailable = FSimpleDelegate());
+	static UE_API void InitPythonAndPipInstall(FSimpleDelegate OnPipInstalled, FSimpleDelegate OnPythonUnavailable = FSimpleDelegate());
+};
+
 #undef UE_API
