PR #13425: PythonScriptPlugin: Add tick_progress to ScopedSlowTask.

diff --git a/Engine/Plugins/Experimental/PythonScriptPlugin/Source/PythonScriptPlugin/Private/PyCore.cpp b/Engine/Plugins/Experimental/PythonScriptPlugin/Source/PythonScriptPlugin/Private/PyCore.cpp
--- a/Engine/Plugins/Experimental/PythonScriptPlugin/Source/PythonScriptPlugin/Private/PyCore.cpp
+++ b/Engine/Plugins/Experimental/PythonScriptPlugin/Source/PythonScriptPlugin/Private/PyCore.cpp
@@ -673,6 +673,18 @@ PyTypeObject InitializePyScopedSlowTaskType()
 			const bool bShouldCancel = InSelf->SlowTask->ShouldCancel();
 			return PyConversion::Pythonize(bShouldCancel);
 		}
+		
+		static PyObject* TickProgress(FPyScopedSlowTask* InSelf)
+		{
+			if (!FPyScopedSlowTask::ValidateInternalState(InSelf))
+			{
+				return nullptr;
+			}
+			
+			InSelf->SlowTask->TickProgress();
+			
+			Py_RETURN_NONE;
+		}
 	};
 
 	static PyMethodDef PyMethods[] = {
@@ -681,7 +693,8 @@ PyTypeObject InitializePyScopedSlowTaskType()
 		{ "make_dialog_delayed", PyCFunctionCast(&FMethods::MakeDialogDelayed), METH_VARARGS | METH_KEYWORDS, "make_dialog_delayed(self, delay: float, can_cancel: bool = False, allow_in_pie: bool = False) -> None -- creates a new dialog for this slow task after the given time delay (in seconds). If the task completes before this time, no dialog will be shown" },
 		{ "make_dialog", PyCFunctionCast(&FMethods::MakeDialog), METH_VARARGS | METH_KEYWORDS, "make_dialog(self, can_cancel: bool = False, allow_in_pie: bool = False) -> None -- creates a new dialog for this slow task, if there is currently not one open" },
 		{ "enter_progress_frame", PyCFunctionCast(&FMethods::EnterProgressFrame), METH_VARARGS | METH_KEYWORDS, "enter_progress_frame(self, work: float = 1.0, desc: Union[Text, str] = \"\") -> None -- indicate that we are to enter a frame that will take up the specified amount of work (completes any previous frames)" },
-		{ "should_cancel", PyCFunctionCast(&FMethods::ShouldCancel), METH_NOARGS, "x.should_cancel() -> bool -- True if the user has requested that the slow task be canceled" },
+		{ "should_cancel", PyCFunctionCast(&FMethods::ShouldCancel), METH_NOARGS, "should_cancel(self) -> bool -- True if the user has requested that the slow task be canceled" },
+		{ "tick_progress", PyCFunctionCast(&FMethods::TickProgress), METH_NOARGS, "tick_progress(self) -> None -- Let the UI refresh but doesn't advance progress" },
 		{ nullptr, nullptr, 0, nullptr }
 	};
 
