Run python startup scripts before commandlets start. Primariy this allows content validators to be registered from python scripts.
If commandlets don't want python to be enabled and startup scripts to run they can use the existing configuration:
[PythonScriptPlugin]
+DisablePythonForCommandlet=

diff --git a/Engine/Plugins/Experimental/PythonScriptPlugin/Source/PythonScriptPlugin/Private/PythonScriptPlugin.cpp b/Engine/Plugins/Experimental/PythonScriptPlugin/Source/PythonScriptPlugin/Private/PythonScriptPlugin.cpp
--- a/Engine/Plugins/Experimental/PythonScriptPlugin/Source/PythonScriptPlugin/Private/PythonScriptPlugin.cpp
+++ b/Engine/Plugins/Experimental/PythonScriptPlugin/Source/PythonScriptPlugin/Private/PythonScriptPlugin.cpp
@@ -1296,6 +1296,16 @@ void FPythonScriptPlugin::InitializePython()
 			Tick(DeltaTime);
 			return true;
 		}));
+
+		if (IsRunningCommandlet())
+		{
+			// Commandlets in general do not tick the engine, so if the plugin is initialized, register a callback to run startup scripts 
+			// before the commandlet executes. 
+			// This allows, for example, creation of content validators in Python that will be executed in a commandlet. 
+			FCoreDelegates::OnCommandletPreMain.AddLambda([this](){
+				RunStartupScripts();
+			});
+		}
 	}
 
 	// Release the GIL taken by Py_Initialize now that initialization has finished, to allow other threads access to Python
