Exposed FTSTicker to Python

This exposes the superior and safer FTSTicker-based ticking method to Python, and existing Python scripts using the Slate-based ticking methods should consider migrating to use this new method (unless forwarding ticks to an embedded UI framework, which was the intent of the Slate-based ticking methods).

The new ticker callables can be registered via `unreal.register_ticker_callback(...)`, which takes a callable (taking a float delta time, and returning a bool where True means continue to tick), along with the delay (in seconds; default 0) that we should apply between each ticker callback. `unreal.unregister_ticker_callback(...)` can be used to remove a previously registered ticker callback.

Examples:
```
def run_once(dt):
print("run_once")
return False # stop ticking
unreal.register_ticker_callback(run_once, 1.0) # call run_once once after 1 second
```
```
def run_many(dt):
print("run_many")
return True # continue ticking
run_many_handle = unreal.register_ticker_callback(run_many, 1.0) # call run_many repeatedly after 1 second delays
```
```
unreal.unregister_ticker_callback(run_many_handle) # stop calling run_many
```

diff --git a/Engine/Plugins/Experimental/PythonScriptPlugin/Source/PythonScriptPlugin/Private/PyCore.cpp b/Engine/Plugins/Experimental/PythonScriptPlugin/Source/PythonScriptPlugin/Private/PyCore.cpp
--- a/Engine/Plugins/Experimental/PythonScriptPlugin/Source/PythonScriptPlugin/Private/PyCore.cpp
+++ b/Engine/Plugins/Experimental/PythonScriptPlugin/Source/PythonScriptPlugin/Private/PyCore.cpp
@@ -35,6 +35,76 @@
 namespace PyCoreUtil
 {
 
+TArray<FTSTicker::FDelegateHandle> PythonTickerCallbackHandles;
+
+FPyTickerDelegateHandle* RegisterTickerCallback(const float InDelay, PyObject* InPyCallable)
+{
+	// We copy the PyCallable into the lambda to keep the Python object alive as long as the delegate is bound
+	FTSTicker::FDelegateHandle TickerDelegateHandle = FTSTicker::GetCoreTicker().AddTicker(TEXT("PythonTicker"), InDelay, [PyCallable = FPyAutoGILPtr(FPyObjectPtr::NewReference(InPyCallable))](const float InDeltaTime) mutable
+	{
+		if (!PyUtil::CanTickPython())
+		{
+			return true;
+		}
+
+		FPyScopedGIL GIL;
+
+		FPyObjectPtr PyArgs = FPyObjectPtr::StealReference(PyTuple_New(1));
+		PyTuple_SetItem(PyArgs, 0, PyConversion::Pythonize(InDeltaTime)); // SetItem steals the reference
+
+		bool bContinue = true;
+		FPyObjectPtr Result = FPyObjectPtr::StealReference(PyObject_CallObject(PyCallable.Get().GetPtr(), PyArgs));
+		if (Result)
+		{
+			PyConversion::Nativize(Result, bContinue, PyConversion::ESetErrorState::No);
+		}
+		else
+		{
+			PyUtil::LogPythonError();
+		}
+
+		return bContinue;
+	});
+
+	if (TickerDelegateHandle.IsValid())
+	{
+		PythonTickerCallbackHandles.Add(TickerDelegateHandle);
+	}
+
+	return FPyTickerDelegateHandle::CreateInstance(TickerDelegateHandle);
+}
+
+bool UnregisterTickerCallback(PyObject* InCallbackHandle)
+{
+	FPyTickerDelegateHandlePtr PyTickerDelegateHandle = FPyTickerDelegateHandlePtr::StealReference(FPyTickerDelegateHandle::CastPyObject(InCallbackHandle));
+	if (!PyTickerDelegateHandle)
+	{
+		return false;
+	}
+
+	FTSTicker::FDelegateHandle& TickerDelegateHandle = PyTickerDelegateHandle->Value;
+	if (TickerDelegateHandle.IsValid())
+	{
+		PythonTickerCallbackHandles.Remove(TickerDelegateHandle);
+		FTSTicker::GetCoreTicker().RemoveTicker(TickerDelegateHandle);
+	}
+
+	return true;
+}
+
+void ClearTickerCallbacks()
+{
+	if (PythonTickerCallbackHandles.Num() > 0)
+	{
+		FTSTicker& CoreTicker = FTSTicker::GetCoreTicker();
+		for (const FTSTicker::FDelegateHandle& Handle : PythonTickerCallbackHandles)
+		{
+			CoreTicker.RemoveTicker(Handle);
+		}
+		PythonTickerCallbackHandles.Empty();
+	}
+}
+
 bool ConvertOptionalString(PyObject* InObj, FString& OutString, const TCHAR* InErrorCtxt, const TCHAR* InErrorMsg)
 {
 	OutString.Reset();
@@ -147,6 +217,32 @@ FPyDelegateHandle* FPyDelegateHandle::CastPyObject(PyObject* InPyObject)
 }
 
 
+FPyTickerDelegateHandle* FPyTickerDelegateHandle::CreateInstance(const FTSTicker::FDelegateHandle& InValue)
+{
+	FPyTickerDelegateHandlePtr NewInstance = FPyTickerDelegateHandlePtr::StealReference(FPyTickerDelegateHandle::New(&PyTickerDelegateHandleType));
+	if (NewInstance)
+	{
+		if (FPyTickerDelegateHandle::Init(NewInstance, InValue) != 0)
+		{
+			PyUtil::LogPythonError();
+			return nullptr;
+		}
+	}
+	return NewInstance.Release();
+}
+
+FPyTickerDelegateHandle* FPyTickerDelegateHandle::CastPyObject(PyObject* InPyObject)
+{
+	if (PyObject_IsInstance(InPyObject, (PyObject*)&PyTickerDelegateHandleType) == 1)
+	{
+		Py_INCREF(InPyObject);
+		return (FPyTickerDelegateHandle*)InPyObject;
+	}
+
+	return nullptr;
+}
+
+
 FPyScopedSlowTask* FPyScopedSlowTask::New(PyTypeObject* InType)
 {
 	FPyScopedSlowTask* Self = (FPyScopedSlowTask*)InType->tp_alloc(InType, 0);
@@ -1045,6 +1141,7 @@ PyTypeObject InitializePyUFunctionDefType()
 
 
 PyTypeObject PyDelegateHandleType = InitializePyWrapperBasicType<FPyDelegateHandle>("_DelegateHandle", "Type for all Unreal exposed FDelegateHandle instances");
+PyTypeObject PyTickerDelegateHandleType = InitializePyWrapperBasicType<FPyTickerDelegateHandle>("_TickerDelegateHandle", "Type for all Unreal exposed FTSTicker::FDelegateHandle instances");
 PyTypeObject PyScopedSlowTaskType = InitializePyScopedSlowTaskType();
 PyTypeObject PyObjectIteratorType = InitializePyObjectIteratorType();
 PyTypeObject PyClassIteratorType = InitializePyTypeIteratorType<UClass, FPyClassIterator>("ClassIterator", "Type for iterating Unreal class types");
@@ -1922,6 +2019,46 @@ PyObject* UnregisterPythonShutdownCallback(PyObject* InSelf, PyObject* InArgs)
 	Py_RETURN_NONE;
 }
 
+PyObject* RegisterTickerCallback(PyObject* InSelf, PyObject* InArgs, PyObject* InKwds)
+{
+	PyObject* PyCallableObj = nullptr;
+	PyObject* PyDelayObj = nullptr;
+
+	static const char* ArgsKwdList[] = { "callable", "delay", nullptr };
+	if (!PyArg_ParseTupleAndKeywords(InArgs, InKwds, "O|O:register_ticker_callback", (char**)ArgsKwdList, &PyCallableObj, &PyDelayObj))
+	{
+		return nullptr;
+	}
+	check(PyCallableObj);
+
+	float Delay = 0.0f;
+	if (PyDelayObj && !PyConversion::Nativize(PyDelayObj, Delay))
+	{
+		PyUtil::SetPythonError(PyExc_TypeError, TEXT("register_ticker_callback"), *FString::Printf(TEXT("Failed to convert 'delay' (%s) to 'float'"), *PyUtil::GetFriendlyTypename(PyDelayObj)));
+		return nullptr;
+	}
+
+	return (PyObject*)PyCoreUtil::RegisterTickerCallback(Delay, PyCallableObj);
+}
+
+PyObject* UnregisterTickerCallback(PyObject* InSelf, PyObject* InArgs)
+{
+	PyObject* PyCallableObj = nullptr;
+	if (!PyArg_ParseTuple(InArgs, "O:unregister_ticker_callback", &PyCallableObj))
+	{
+		return nullptr;
+	}
+	check(PyCallableObj);
+
+	if (!PyCoreUtil::UnregisterTickerCallback(PyCallableObj))
+	{
+		PyUtil::SetPythonError(PyExc_TypeError, TEXT("unregister_ticker_callback"), *FString::Printf(TEXT("Failed to convert argument '%s' to '_TickerDelegateHandle'"), *PyUtil::GetFriendlyTypename(PyCallableObj)));
+		return nullptr;
+	}
+
+	Py_RETURN_NONE;
+}
+
 // NOTE: Several functions like 'new_object' are marked as returning 'Any' because the exact return type cannot be deduced from the Class type. We could mark them as returning base 'Object'
 //       type, but that would force user to cast down to please type checker while the returned type is already the exact requested type. For example, the instruction
 //       'o: PyTestObject = unreal.new_object(unreal.PyTestObject.static_class())' must pass type checker and if new_object() was to return Object, it would not be automatically cast down.
@@ -1959,6 +2096,8 @@ PyMethodDef PyCoreMethods[] = {
 	{ "create_python_object_handle", PyCFunctionCast(&CreatePythonObjectHandle), METH_VARARGS, "create_python_object_handle(obj: Optional[Any]) -> Optional[PythonObjectHandle] -- create a PythonObjectHandle from the given PyObject" },
 	{ "resolve_python_object_handle", PyCFunctionCast(&ResolvePythonObjectHandle), METH_VARARGS, "resolve_python_object_handle(handle: Optional[PythonObjectHandle]) -> Optional[Any] -- resolve a PythonObjectHandle to its PyObject, or None" },
 	{ "destroy_python_object_handle", PyCFunctionCast(&DestroyPythonObjectHandle), METH_VARARGS, "destroy_python_object_handle(handle: Optional[PythonObjectHandle]) -- destroy a PythonObjectHandle (clearing its PyObject reference)" },
+	{ "register_ticker_callback", PyCFunctionCast(&RegisterTickerCallback), METH_VARARGS | METH_KEYWORDS, "x.register_ticker_callback(callable: Union[Callable[[float], None]->bool, DelegateBase], delay: float=0.0) -> object -- register the given callable (taking a single float and returning a bool where true will continue ticking) as a ticker callback in FTSTicker, to run at most every delay seconds" },
+	{ "unregister_ticker_callback", PyCFunctionCast(&UnregisterTickerCallback), METH_VARARGS, "x.unregister_ticker_callback(handle: object) -> None -- unregister the given handle from a previous call to register_ticker_callback" },
 	{ nullptr, nullptr, 0, nullptr }
 };
 
@@ -1973,6 +2112,7 @@ void InitializeModule()
 	NativePythonModule.PyModule = PyImport_AddModule("_unreal_core");
 	PyModule_AddFunctions(NativePythonModule.PyModule, PyCoreMethods);
 	PyType_Ready(&PyDelegateHandleType);
+	PyType_Ready(&PyTickerDelegateHandleType);
 
 	if (PyType_Ready(&PyScopedSlowTaskType) == 0)
 	{
@@ -2033,6 +2173,7 @@ void InitializeModule()
 
 void ShutdownModule()
 {
+	PyCoreUtil::ClearTickerCallbacks();
 }
 
 FPythonLogCapture& GetPythonLogCapture()
diff --git a/Engine/Plugins/Experimental/PythonScriptPlugin/Source/PythonScriptPlugin/Private/PyCore.h b/Engine/Plugins/Experimental/PythonScriptPlugin/Source/PythonScriptPlugin/Private/PyCore.h
--- a/Engine/Plugins/Experimental/PythonScriptPlugin/Source/PythonScriptPlugin/Private/PyCore.h
+++ b/Engine/Plugins/Experimental/PythonScriptPlugin/Source/PythonScriptPlugin/Private/PyCore.h
@@ -7,7 +7,7 @@
 #include "PyUtil.h"
 #include "PyWrapperBasic.h"
 #include "PyPtr.h"
-#include "CoreMinimal.h"
+#include "Containers/Ticker.h"
 #include "Misc/EnumClassFlags.h"
 #include "UObject/UObjectIterator.h"
 #include "Delegates/DelegateCombinations.h"
@@ -34,6 +34,9 @@ UObjectRedirector* CreatePythonTypeLegacyRedirector(const FString& ShortName, co
 /** Python type for FPyDelegateHandle */
 extern PyTypeObject PyDelegateHandleType;
 
+/** Python type for FPyTickerDelegateHandle */
+extern PyTypeObject PyTickerDelegateHandleType;
+
 /** Python type for FPyScopedSlowTask */
 extern PyTypeObject PyScopedSlowTaskType;
 
@@ -70,6 +73,18 @@ struct FPyDelegateHandle : public TPyWrapperBasic<FDelegateHandle, FPyDelegateHa
 	static FPyDelegateHandle* CastPyObject(PyObject* InPyObject);
 };
 
+/** Type for all Unreal exposed FTSTicker::FDelegateHandle instances */
+struct FPyTickerDelegateHandle : public TPyWrapperBasic<FTSTicker::FDelegateHandle, FPyTickerDelegateHandle>
+{
+	typedef TPyWrapperBasic<FTSTicker::FDelegateHandle, FPyTickerDelegateHandle> Super;
+
+	/** Create and initialize a new wrapper instance from the given native instance */
+	static FPyTickerDelegateHandle* CreateInstance(const FTSTicker::FDelegateHandle& InValue);
+
+	/** Cast the given Python object to this wrapped type (returns a new reference) */
+	static FPyTickerDelegateHandle* CastPyObject(PyObject* InPyObject);
+};
+
 /** Type used to create and managed a scoped slow task in Python */
 struct FPyScopedSlowTask
 {
@@ -368,6 +383,7 @@ struct FPyUFunctionDef
 };
 
 typedef TPyPtr<FPyDelegateHandle> FPyDelegateHandlePtr;
+typedef TPyPtr<FPyTickerDelegateHandle> FPyTickerDelegateHandlePtr;
 typedef TPyPtr<FPyScopedSlowTask> FPyScopedSlowTaskPtr;
 typedef TPyPtr<FPyObjectIterator> FPyObjectIteratorPtr;
 typedef TPyPtr<FPyClassIterator> FPyClassIteratorPtr;
diff --git a/Engine/Plugins/Experimental/PythonScriptPlugin/Source/PythonScriptPlugin/Private/PySlate.cpp b/Engine/Plugins/Experimental/PythonScriptPlugin/Source/PythonScriptPlugin/Private/PySlate.cpp
--- a/Engine/Plugins/Experimental/PythonScriptPlugin/Source/PythonScriptPlugin/Private/PySlate.cpp
+++ b/Engine/Plugins/Experimental/PythonScriptPlugin/Source/PythonScriptPlugin/Private/PySlate.cpp
@@ -8,7 +8,6 @@
 #include "PyWrapperTypeRegistry.h"
 
 #include "UObject/Package.h"
-#include "UObject/UObjectThreadContext.h"
 #include "Framework/Application/SlateApplication.h"
 
 #if PLATFORM_WINDOWS
@@ -29,15 +28,7 @@ FPyDelegateHandle* RegisterSlateTickCallback(FSlateApplication::FSlateTickEvent&
 	// We copy the PyCallable into the lambda to keep the Python object alive as long as the delegate is bound
 	FDelegateHandle TickEventDelegateHandle = InSlateTickEvent.AddLambda([PyCallable = FPyAutoGILPtr(FPyObjectPtr::NewReference(InPyCallable))](const float InDeltaTime) mutable
 	{
-		// Do not tick into Python while Python is already running on the stack
-		// Doing so could cause re-entrancy issues and confuse exception handling
-		if (UE::Python::Private::GIsRunningUserScript)
-		{
-			return;
-		}
-
-		// Do not tick into Python when it may not be safe to call back into C++
-		if (GIsSavingPackage || IsGarbageCollecting() || FUObjectThreadContext::Get().IsRoutingPostLoad)
+		if (!PyUtil::CanTickPython())
 		{
 			return;
 		}
diff --git a/Engine/Plugins/Experimental/PythonScriptPlugin/Source/PythonScriptPlugin/Private/PyUtil.cpp b/Engine/Plugins/Experimental/PythonScriptPlugin/Source/PythonScriptPlugin/Private/PyUtil.cpp
--- a/Engine/Plugins/Experimental/PythonScriptPlugin/Source/PythonScriptPlugin/Private/PyUtil.cpp
+++ b/Engine/Plugins/Experimental/PythonScriptPlugin/Source/PythonScriptPlugin/Private/PyUtil.cpp
@@ -37,6 +37,7 @@
 #include "UObject/StructOnScope.h"
 #include "UObject/TextProperty.h"
 #include "UObject/UnrealType.h"
+#include "UObject/UObjectThreadContext.h"
 
 #if WITH_EDITOR
 #include "EditorSubsystem.h"
@@ -1748,6 +1749,24 @@ void CollectGarbage(int32 Generations)
 	}
 }
 
+bool CanTickPython()
+{
+	// Do not tick into Python while Python is already running on the stack
+	// Doing so could cause re-entrancy issues and confuse exception handling
+	if (UE::Python::Private::GIsRunningUserScript)
+	{
+		return false;
+	}
+
+	// Do not tick into Python when it may not be safe to call back into C++
+	if (GIsSavingPackage || IsGarbageCollecting() || FUObjectThreadContext::Get().IsRoutingPostLoad)
+	{
+		return false;
+	}
+
+	return true;
+}
+
 }
 
 #endif	// WITH_PYTHON
diff --git a/Engine/Plugins/Experimental/PythonScriptPlugin/Source/PythonScriptPlugin/Private/PyUtil.h b/Engine/Plugins/Experimental/PythonScriptPlugin/Source/PythonScriptPlugin/Private/PyUtil.h
--- a/Engine/Plugins/Experimental/PythonScriptPlugin/Source/PythonScriptPlugin/Private/PyUtil.h
+++ b/Engine/Plugins/Experimental/PythonScriptPlugin/Source/PythonScriptPlugin/Private/PyUtil.h
@@ -681,6 +681,9 @@ namespace PyUtil
 	 * @see [REDACTED_PROFILE_URL]/cpython/blob/main/InternalDocs/garbage_collector.md#optimization-generations
 	 */
 	void CollectGarbage(int32 Generations = 2);
+
+	/** Return true if it's currently considered safe to invoke Python from tick handlers */
+	bool CanTickPython();
 }
 
 #endif	// WITH_PYTHON
