Removed extra Python object when converting PyUnicode to FString

#rnx

diff --git a/Engine/Plugins/Experimental/PythonScriptPlugin/Source/PythonScriptPlugin/Private/PyConversion.cpp b/Engine/Plugins/Experimental/PythonScriptPlugin/Source/PythonScriptPlugin/Private/PyConversion.cpp
--- a/Engine/Plugins/Experimental/PythonScriptPlugin/Source/PythonScriptPlugin/Private/PyConversion.cpp
+++ b/Engine/Plugins/Experimental/PythonScriptPlugin/Source/PythonScriptPlugin/Private/PyConversion.cpp
@@ -307,10 +307,8 @@ FPyConversionResult Nativize(PyObject* PyObj, FString& OutVal, const ESetErrorSt
 {
 	if (PyUnicode_Check(PyObj))
 	{
-		FPyObjectPtr PyBytesObj = FPyObjectPtr::StealReference(PyUnicode_AsUTF8String(PyObj));
-		if (PyBytesObj)
+		if (const char* PyUtf8Buffer = PyUnicode_AsUTF8(PyObj))
 		{
-			const char* PyUtf8Buffer = PyBytes_AsString(PyBytesObj);
 			OutVal = UTF8_TO_TCHAR(PyUtf8Buffer);
 			return FPyConversionResult::Success();
 		}
