Allow UE to run if Python is not able to load properly

diff --git a/Engine/Plugins/Experimental/PythonScriptPlugin/Source/PythonScriptPlugin/Private/PyConstant.cpp b/Engine/Plugins/Experimental/PythonScriptPlugin/Source/PythonScriptPlugin/Private/PyConstant.cpp
--- a/Engine/Plugins/Experimental/PythonScriptPlugin/Source/PythonScriptPlugin/Private/PyConstant.cpp
+++ b/Engine/Plugins/Experimental/PythonScriptPlugin/Source/PythonScriptPlugin/Private/PyConstant.cpp
@@ -6,9 +6,15 @@
 
 #if WITH_PYTHON
 
-void InitializePyConstant()
+bool InitializePyConstant()
 {
-	PyType_Ready(&PyConstantDescrType);
+	if (!ensure(PyType_Ready(&PyConstantDescrType) == 0))
+	{
+		PyUtil::LogPythonError();
+		UE_LOG(LogPython, Error, TEXT("PyType_Ready(&PyConstantDescrType) failed! Python will be disabled."))
+		return false;
+	}
+	return true;
 }
 
 bool FPyConstantDef::AddConstantToType(FPyConstantDef* InConstant, PyTypeObject* InType)
diff --git a/Engine/Plugins/Experimental/PythonScriptPlugin/Source/PythonScriptPlugin/Private/PyConstant.h b/Engine/Plugins/Experimental/PythonScriptPlugin/Source/PythonScriptPlugin/Private/PyConstant.h
--- a/Engine/Plugins/Experimental/PythonScriptPlugin/Source/PythonScriptPlugin/Private/PyConstant.h
+++ b/Engine/Plugins/Experimental/PythonScriptPlugin/Source/PythonScriptPlugin/Private/PyConstant.h
@@ -11,7 +11,7 @@
 extern PyTypeObject PyConstantDescrType;
 
 /** Initialize the PyConstant types */
-void InitializePyConstant();
+bool InitializePyConstant();
 
 /** Function pointer used to convert a constant value pointer into a Python object */
 typedef PyObject*(*PyConstantGetter)(PyTypeObject*, const void*);
diff --git a/Engine/Plugins/Experimental/PythonScriptPlugin/Source/PythonScriptPlugin/Private/PyMethodWithClosure.cpp b/Engine/Plugins/Experimental/PythonScriptPlugin/Source/PythonScriptPlugin/Private/PyMethodWithClosure.cpp
--- a/Engine/Plugins/Experimental/PythonScriptPlugin/Source/PythonScriptPlugin/Private/PyMethodWithClosure.cpp
+++ b/Engine/Plugins/Experimental/PythonScriptPlugin/Source/PythonScriptPlugin/Private/PyMethodWithClosure.cpp
@@ -61,13 +61,31 @@ private:
 
 FPyCFunctionWithClosureObjectFreeList* PyCFunctionWithClosureObjectFreeList = nullptr;
 
-void InitializePyMethodWithClosure()
+bool InitializePyMethodWithClosure()
 {
-	PyType_Ready(&PyCFunctionWithClosureType);
-	PyType_Ready(&PyMethodWithClosureDescrType);
-	PyType_Ready(&PyClassMethodWithClosureDescrType);
+	if (!ensure(PyType_Ready(&PyCFunctionWithClosureType) == 0))
+	{
+		PyUtil::LogPythonError();
+		UE_LOG(LogPython, Error, TEXT("PyType_Ready(&PyCFunctionWithClosureType) failed! Python will be disabled."))
+		return false;
+	}
+
+	if (!ensure(PyType_Ready(&PyMethodWithClosureDescrType) == 0))
+	{
+		PyUtil::LogPythonError();
+		UE_LOG(LogPython, Error, TEXT("PyType_Ready(&PyMethodWithClosureDescrType) failed! Python will be disabled."))
+		return false;
+	}
+	
+	if (!ensure(PyType_Ready(&PyClassMethodWithClosureDescrType) == 0))
+	{
+		PyUtil::LogPythonError();
+		UE_LOG(LogPython, Error, TEXT("PyType_Ready(&PyClassMethodWithClosureDescrType) failed! Python will be disabled."))
+		return false;
+	}
 
 	PyCFunctionWithClosureObjectFreeList = new FPyCFunctionWithClosureObjectFreeList();
+	return true;
 }
 
 void ShutdownPyMethodWithClosure()
diff --git a/Engine/Plugins/Experimental/PythonScriptPlugin/Source/PythonScriptPlugin/Private/PyMethodWithClosure.h b/Engine/Plugins/Experimental/PythonScriptPlugin/Source/PythonScriptPlugin/Private/PyMethodWithClosure.h
--- a/Engine/Plugins/Experimental/PythonScriptPlugin/Source/PythonScriptPlugin/Private/PyMethodWithClosure.h
+++ b/Engine/Plugins/Experimental/PythonScriptPlugin/Source/PythonScriptPlugin/Private/PyMethodWithClosure.h
@@ -17,7 +17,7 @@ extern PyTypeObject PyMethodWithClosureDescrType;
 extern PyTypeObject PyClassMethodWithClosureDescrType;
 
 /** Initialize the PyMethodWithClosure types */
-void InitializePyMethodWithClosure();
+bool InitializePyMethodWithClosure();
 
 /** Shutdown the PyMethodWithClosure types */
 void ShutdownPyMethodWithClosure();
diff --git a/Engine/Plugins/Experimental/PythonScriptPlugin/Source/PythonScriptPlugin/Private/PythonScriptPlugin.cpp b/Engine/Plugins/Experimental/PythonScriptPlugin/Source/PythonScriptPlugin/Private/PythonScriptPlugin.cpp
--- a/Engine/Plugins/Experimental/PythonScriptPlugin/Source/PythonScriptPlugin/Private/PythonScriptPlugin.cpp
+++ b/Engine/Plugins/Experimental/PythonScriptPlugin/Source/PythonScriptPlugin/Private/PythonScriptPlugin.cpp
@@ -17,6 +17,7 @@
 #include "EngineAnalytics.h"
 #include "Interfaces/IAnalyticsProvider.h"
 #include "AssetRegistry/AssetRegistryModule.h"
+#include "Logging/MessageLog.h"
 #include "Modules/ModuleManager.h"
 #include "UObject/PackageReload.h"
 #include "Misc/App.h"
@@ -170,6 +171,18 @@ bool IsPythonEnabled()
 	return false;
 }
 
+void ShowPythonDisabledNotification()
+{
+	FMessageLog MessageLog("EditorErrors");
+	MessageLog.SetCurrentPage(LOCTEXT("MessageLogPage", "Python"));
+	
+	const FText ErrorMessage = LOCTEXT("PythonCouldNotBeInitialized_Text", "Python could not be initialized and was disabled");
+
+	MessageLog.Error(ErrorMessage);
+	MessageLog.Info(LOCTEXT("PythonCouldNotBeInitialized_SubText", "See OutputLog for more information"));
+	MessageLog.Notify(ErrorMessage, EMessageSeverity::Error, true /*bForce*/);
+}
+
 } // namespace UE::Python::Private
 
 /** Util struct to set the sys.argv data for Python when executing a file with arguments */
@@ -962,7 +975,16 @@ void FPythonScriptPlugin::ConfigureAndInitializePython()
 		return;
 	}
 
-	InitializePython();
+	if (!InitializePython())
+	{
+		bIsEnabled = false;
+		UE_LOG(LogPython, Error, TEXT("Python could not be initialized and was disabled."))
+		if (!FApp::IsUnattended() && !IsRunningCommandlet() && !GIsRunningUnattendedScript)
+		{
+			FDelayedAutoRegisterHelper(EDelayedRegisterRunPhase::EndOfEngineInit, &UE::Python::Private::ShowPythonDisabledNotification);
+		}
+		return;
+	}
 
 	IModularFeatures::Get().RegisterModularFeature(IConsoleCommandExecutor::ModularFeatureName(), &CmdExec);
 	IModularFeatures::Get().RegisterModularFeature(IConsoleCommandExecutor::ModularFeatureName(), &CmdREPLExec);
@@ -1053,7 +1075,7 @@ void FPythonScriptPlugin::PostChange(const UUserDefinedEnum* Enum, FEnumEditorUt
 
 #if WITH_PYTHON
 
-void FPythonScriptPlugin::InitializePython()
+bool FPythonScriptPlugin::InitializePython()
 {
 	TRACE_CPUPROFILER_EVENT_SCOPE(FPythonScriptPlugin::InitializePython)
 
@@ -1105,7 +1127,13 @@ void FPythonScriptPlugin::InitializePython()
 		PreConfig.isolated = IsolatedInterpreterFlag;
 		PreConfig.use_environment = !IsolatedInterpreterFlag;
 
-		Py_PreInitialize(&PreConfig);
+		if (PyStatus Status = Py_PreInitialize(&PreConfig);
+			!ensure(!PyStatus_Exception(Status)))
+		{
+			PyUtil::LogPythonError();
+			UE_LOG(LogPython, Error, TEXT("Py_PreInitialize failed! Error: '%hs'. Python will be disabled."), Status.err_msg)
+			return false;
+		}
 
 		// Create empty init config
 		PyConfig_InitIsolatedConfig(&ModulePyConfig);
@@ -1128,7 +1156,13 @@ void FPythonScriptPlugin::InitializePython()
 		ModulePyConfig.install_signal_handlers = 0;
 		ModulePyConfig.safe_path = 0;
 
-		Py_InitializeFromConfig(&ModulePyConfig);
+		if (PyStatus Status = Py_InitializeFromConfig(&ModulePyConfig);
+			!ensure(!PyStatus_Exception(Status)))
+		{
+			PyUtil::LogPythonError();
+			UE_LOG(LogPython, Error, TEXT("Py_InitializeFromConfig failed! Error: '%hs'. Python will be disabled."), Status.err_msg)
+			return false;
+		}
 #else
 		Py_SetProgramName(PyProgramName.GetData());
 		Py_SetPythonHome(PyHomePath.GetData());
@@ -1152,7 +1186,12 @@ void FPythonScriptPlugin::InitializePython()
 #if PY_MAJOR_VERSION >=3 && PY_MINOR_VERSION >= 11
 		// Set default argv to [""]
 		FPyObjectPtr PyArgvList = UE::Python::Private::MakeEmptyArgvList();
-		PySys_SetObject("argv", PyArgvList.Get());
+		if (!ensure(PySys_SetObject("argv", PyArgvList.Get()) == 0))
+		{
+			PyUtil::LogPythonError();
+			UE_LOG(LogPython, Error, TEXT("PySys_SetObject(\"argv\") failed! Python will be disabled."))
+			return false;
+		}
 #else
 		PySys_SetArgvEx(1, UE::Python::Private::NullPyArgPtrs, 0);
 #endif
@@ -1161,22 +1200,33 @@ void FPythonScriptPlugin::InitializePython()
 		if (IsDeveloperModeEnabled())
 		{
 			PyUtil::EnableDeveloperWarnings();
+			ensure(!PyUtil::LogPythonError());
 		}
 
 		// Check if the user wants type hinting. (In the stub and/or Docstrings).
 		PyGenUtil::SetTypeHintingMode(GetTypeHintingMode());
 
 		// Initialize our custom method type as we'll need it when generating bindings
-		InitializePyMethodWithClosure();
+		if (!ensure(InitializePyMethodWithClosure()))
+		{
+			return false;
+		}
 
 		// Initialize our custom constant type as we'll need it when generating bindings
-		InitializePyConstant();
+		if (!ensure(InitializePyConstant()))
+		{
+			return false;
+		}
+
+		// Some users have PyImport_AddModule return null below, so we want to make sure errors are cleared before calling it
+		ensure(!PyUtil::LogPythonError());
 
 		PyObject* PyMainModule = PyImport_AddModule("__main__");
-		if (!PyMainModule)
+		if (!ensure(PyMainModule))
 		{
 			PyUtil::LogPythonError();
-			checkf(false, TEXT("Failed to initialize '__main__'! This typically means the Python SDK is missing or could not be loaded"));
+			UE_LOG(LogPython, Error, TEXT("Failed to initialize '__main__'! This typically means the Python SDK is missing or could not be loaded"))
+			return false;
 		}
 		PyDefaultGlobalDict = FPyObjectPtr::NewReference(PyModule_GetDict(PyMainModule));
 		PyDefaultLocalDict = PyDefaultGlobalDict;
@@ -1366,6 +1416,7 @@ void FPythonScriptPlugin::InitializePython()
 		}
 	}
 #endif	// WITH_EDITOR
+	return true;
 }
 
 void FPythonScriptPlugin::ShutdownPython()
diff --git a/Engine/Plugins/Experimental/PythonScriptPlugin/Source/PythonScriptPlugin/Private/PythonScriptPlugin.h b/Engine/Plugins/Experimental/PythonScriptPlugin/Source/PythonScriptPlugin/Private/PythonScriptPlugin.h
--- a/Engine/Plugins/Experimental/PythonScriptPlugin/Source/PythonScriptPlugin/Private/PythonScriptPlugin.h
+++ b/Engine/Plugins/Experimental/PythonScriptPlugin/Source/PythonScriptPlugin/Private/PythonScriptPlugin.h
@@ -163,7 +163,7 @@ private:
 
 	void ConfigureAndInitializePython();
 
-	void InitializePython();
+	bool InitializePython();
 
 	void ShutdownPython();
 
