Fix for crash in Python module during Shutdown when Python failed initializing

diff --git a/Engine/Plugins/Experimental/PythonScriptPlugin/Source/PythonScriptPlugin/Private/PythonScriptPlugin.cpp b/Engine/Plugins/Experimental/PythonScriptPlugin/Source/PythonScriptPlugin/Private/PythonScriptPlugin.cpp
--- a/Engine/Plugins/Experimental/PythonScriptPlugin/Source/PythonScriptPlugin/Private/PythonScriptPlugin.cpp
+++ b/Engine/Plugins/Experimental/PythonScriptPlugin/Source/PythonScriptPlugin/Private/PythonScriptPlugin.cpp
@@ -1426,9 +1426,12 @@ void FPythonScriptPlugin::ShutdownPython()
 		return;
 	}
 
-	// We need to restore the original GIL prior to calling Py_Finalize
-	PyEval_RestoreThread(PyMainThreadState);
-	PyMainThreadState = nullptr;
+	if (PyMainThreadState)
+	{
+		// We need to restore the original GIL prior to calling Py_Finalize
+		PyEval_RestoreThread(PyMainThreadState);
+		PyMainThreadState = nullptr;
+	}
 
 #if WITH_EDITOR
 	// Remove the Content Browser integration
@@ -1487,7 +1490,10 @@ void FPythonScriptPlugin::ShutdownPython()
 
 	ShutdownPyMethodWithClosure();
 
-	Py_Finalize();
+	if (Py_IsInitialized())
+	{
+		Py_Finalize();
+	}
 
 	bIsConfigured = false;
 	bIsInterpreterInitialized = false;
