Fixed incorrect type hinting for unreal.register_ticker_callback

#rnx

diff --git a/Engine/Plugins/Experimental/PythonScriptPlugin/Source/PythonScriptPlugin/Private/PyCore.cpp b/Engine/Plugins/Experimental/PythonScriptPlugin/Source/PythonScriptPlugin/Private/PyCore.cpp
--- a/Engine/Plugins/Experimental/PythonScriptPlugin/Source/PythonScriptPlugin/Private/PyCore.cpp
+++ b/Engine/Plugins/Experimental/PythonScriptPlugin/Source/PythonScriptPlugin/Private/PyCore.cpp
@@ -2096,7 +2096,7 @@ PyMethodDef PyCoreMethods[] = {
 	{ "create_python_object_handle", PyCFunctionCast(&CreatePythonObjectHandle), METH_VARARGS, "create_python_object_handle(obj: Optional[Any]) -> Optional[PythonObjectHandle] -- create a PythonObjectHandle from the given PyObject" },
 	{ "resolve_python_object_handle", PyCFunctionCast(&ResolvePythonObjectHandle), METH_VARARGS, "resolve_python_object_handle(handle: Optional[PythonObjectHandle]) -> Optional[Any] -- resolve a PythonObjectHandle to its PyObject, or None" },
 	{ "destroy_python_object_handle", PyCFunctionCast(&DestroyPythonObjectHandle), METH_VARARGS, "destroy_python_object_handle(handle: Optional[PythonObjectHandle]) -- destroy a PythonObjectHandle (clearing its PyObject reference)" },
-	{ "register_ticker_callback", PyCFunctionCast(&RegisterTickerCallback), METH_VARARGS | METH_KEYWORDS, "x.register_ticker_callback(callable: Union[Callable[[float], None]->bool, DelegateBase], delay: float=0.0) -> object -- register the given callable (taking a single float and returning a bool where true will continue ticking) as a ticker callback in FTSTicker, to run at most every delay seconds" },
+	{ "register_ticker_callback", PyCFunctionCast(&RegisterTickerCallback), METH_VARARGS | METH_KEYWORDS, "x.register_ticker_callback(callable: Union[Callable[[float], bool], DelegateBase], delay: float=0.0) -> object -- register the given callable (taking a single float and returning a bool where true will continue ticking) as a ticker callback in FTSTicker, to run at most every delay seconds" },
 	{ "unregister_ticker_callback", PyCFunctionCast(&UnregisterTickerCallback), METH_VARARGS, "x.unregister_ticker_callback(handle: object) -> None -- unregister the given handle from a previous call to register_ticker_callback" },
 	{ nullptr, nullptr, 0, nullptr }
 };
