Ensure the PythonScript commandlet enables source control prior to running the script

Source control cannot be directly enabled in Python, so it's simpler to enable it from C++ before running the script

diff --git a/Engine/Plugins/Experimental/PythonScriptPlugin/Source/PythonScriptPlugin/Private/PythonScriptCommandlet.cpp b/Engine/Plugins/Experimental/PythonScriptPlugin/Source/PythonScriptPlugin/Private/PythonScriptCommandlet.cpp
--- a/Engine/Plugins/Experimental/PythonScriptPlugin/Source/PythonScriptPlugin/Private/PythonScriptCommandlet.cpp
+++ b/Engine/Plugins/Experimental/PythonScriptPlugin/Source/PythonScriptPlugin/Private/PythonScriptCommandlet.cpp
@@ -4,6 +4,9 @@
 #include "PythonScriptPlugin.h"
 #include "Logging/LogMacros.h"
 #include "Containers/Ticker.h"
+#if WITH_EDITOR
+#include "SourceControlHelpers.h"
+#endif
 
 #include UE_INLINE_GENERATED_CPP_BY_NAME(PythonScriptCommandlet)
 
@@ -49,6 +52,11 @@ int32 UPythonScriptCommandlet::Main(const FString& Params)
 	FPythonScriptPlugin::Get()->RunStartupScripts();
 
 	{
+#if WITH_EDITOR
+		// Ensure that source control is enabled, as commandlets aren't guaranteed to do that
+		FScopedSourceControl SourceControl;
+#endif
+
 		UE_LOG(LogPythonScriptCommandlet, Display, TEXT("Running Python script: %s"), *PythonScript);
 
 		FPythonCommandEx PythonCommand;
diff --git a/Engine/Plugins/Experimental/PythonScriptPlugin/Source/PythonScriptPlugin/PythonScriptPlugin.Build.cs b/Engine/Plugins/Experimental/PythonScriptPlugin/Source/PythonScriptPlugin/PythonScriptPlugin.Build.cs
--- a/Engine/Plugins/Experimental/PythonScriptPlugin/Source/PythonScriptPlugin/PythonScriptPlugin.Build.cs
+++ b/Engine/Plugins/Experimental/PythonScriptPlugin/Source/PythonScriptPlugin/PythonScriptPlugin.Build.cs
@@ -51,6 +51,7 @@ namespace UnrealBuildTool.Rules
 						"AssetTools",
 						"ContentBrowserData",
 						"ContentBrowserFileDataSource",
+						"SourceControl",
 					}
 				);
 			}
