Python add support for delayed iniatialization.

#Synced-CL 48734843

diff --git a/Engine/Plugins/Experimental/PythonScriptPlugin/Source/PythonScriptPlugin/Private/PythonScriptPlugin.cpp b/Engine/Plugins/Experimental/PythonScriptPlugin/Source/PythonScriptPlugin/Private/PythonScriptPlugin.cpp
--- a/Engine/Plugins/Experimental/PythonScriptPlugin/Source/PythonScriptPlugin/Private/PythonScriptPlugin.cpp
+++ b/Engine/Plugins/Experimental/PythonScriptPlugin/Source/PythonScriptPlugin/Private/PythonScriptPlugin.cpp
@@ -94,8 +94,11 @@ FPyObjectPtr MakeEmptyArgvList()
 	return PyArgvList;
 }
 
+static bool bIsPythonLifeCycleDynamic = false;
+static FAutoConsoleVariableRef CVarIsPythonLifeCycleDynamic(TEXT("Engine.Python.IsPythonLifeCycleDynamic"), bIsPythonLifeCycleDynamic, TEXT("If true python will not be initialized on boot in the editor. This allow the project to control when python should be initialized or not."), ECVF_ReadOnly);
+
 static bool bIsEnabledByDefault = true;
-static FAutoConsoleVariableRef CVarAsyncLoadLocalizationData(TEXT("Engine.Python.IsEnabledByDefault"), bIsEnabledByDefault, TEXT("True if Python is enabled by default (after checking any command line overrides), or False if it should be considered disabled. This must be set prior to OnPostEngineInit."));
+static FAutoConsoleVariableRef CVarIsPythonEnabledByDefault(TEXT("Engine.Python.IsEnabledByDefault"), bIsEnabledByDefault, TEXT("True if Python is enabled by default (after checking any command line overrides), or False if it should be considered disabled. This must be set prior to OnPostEngineInit."));
 
 bool IsPythonEnabled()
 {
@@ -970,6 +973,11 @@ void FPythonScriptPlugin::ConfigureAndInitializePython()
 	OnPythonConfiguredDelegate.Broadcast();
 	OnPythonConfiguredDelegate.Clear();
 
+	if (UE::Python::Private::bIsPythonLifeCycleDynamic && !bIsForceEnabledAtRuntime)
+	{
+		return;
+	}
+
 	if (!IsPythonAvailable())
 	{
 		return;
