Use Python docstring to set ToolTip metadata on Python-generated UClasses, just as is already done for UFunctions.

#rnx
#Synced-CL 49066509

diff --git a/Engine/Plugins/Experimental/PythonScriptPlugin/Source/PythonScriptPlugin/Private/PyWrapperObject.cpp b/Engine/Plugins/Experimental/PythonScriptPlugin/Source/PythonScriptPlugin/Private/PyWrapperObject.cpp
--- a/Engine/Plugins/Experimental/PythonScriptPlugin/Source/PythonScriptPlugin/Private/PyWrapperObject.cpp
+++ b/Engine/Plugins/Experimental/PythonScriptPlugin/Source/PythonScriptPlugin/Private/PyWrapperObject.cpp
@@ -1411,6 +1411,13 @@ public:
 			return nullptr;
 		}
 
+		// Generate ToolTip using the Python doc string
+		if (PyType->tp_doc)
+		{
+			static const FName ToolTipKey = "ToolTip";
+			NewClass->SetMetaData(ToolTipKey, UTF8_TO_TCHAR(PyType->tp_doc));
+		}
+
 		// Let Python know that we've changed its type
 		PyType_Modified(PyType);
 
