Python dynamic life cycle.

Python can now be shut down at runtime.
Clean up a bit of statics that weren't cleaned between shutdown and reinit.
When shuting down python early we now run a force clean of the python object that were hold by other systems.

#Synced-CL 49167278

diff --git a/Engine/Plugins/Experimental/PythonScriptPlugin/Source/PythonScriptPlugin/Private/PyCore.cpp b/Engine/Plugins/Experimental/PythonScriptPlugin/Source/PythonScriptPlugin/Private/PyCore.cpp
--- a/Engine/Plugins/Experimental/PythonScriptPlugin/Source/PythonScriptPlugin/Private/PyCore.cpp
+++ b/Engine/Plugins/Experimental/PythonScriptPlugin/Source/PythonScriptPlugin/Private/PyCore.cpp
@@ -2174,6 +2174,18 @@ void InitializeModule()
 void ShutdownModule()
 {
 	PyCoreUtil::ClearTickerCallbacks();
+	if (GetPythonTypeContainerSingleton().IsValid())
+	{
+		UObject* PythonTypesPackage = GetPythonTypeContainerSingleton().Get();
+		PythonTypesPackage->Rename(TEXT("PythonTypes_Garbage"), nullptr, REN_DoNotDirty | REN_DontCreateRedirectors | REN_NonTransactional);
+		PythonTypesPackage->MarkAsGarbage();
+		if (PythonTypesPackage->IsRooted())
+		{
+			PythonTypesPackage->RemoveFromRoot();
+		}
+		PythonTypesPackage->ClearFlags(RF_Public | RF_Standalone);
+		GetPythonTypeContainerSingleton().Reset();
+	}
 }
 
 FPythonLogCapture& GetPythonLogCapture()
diff --git a/Engine/Plugins/Experimental/PythonScriptPlugin/Source/PythonScriptPlugin/Private/PyReferenceCollector.cpp b/Engine/Plugins/Experimental/PythonScriptPlugin/Source/PythonScriptPlugin/Private/PyReferenceCollector.cpp
--- a/Engine/Plugins/Experimental/PythonScriptPlugin/Source/PythonScriptPlugin/Private/PyReferenceCollector.cpp
+++ b/Engine/Plugins/Experimental/PythonScriptPlugin/Source/PythonScriptPlugin/Private/PyReferenceCollector.cpp
@@ -13,6 +13,11 @@
 #include "UObject/UObjectHash.h"
 #include "UObject/PurgingReferenceCollector.h"
 
+#if WITH_EDITOR
+#include "CoreGlobals.h"
+#include "ObjectTools.h"
+#endif
+
 #if WITH_PYTHON
 
 FPyReferenceCollector& FPyReferenceCollector::Get()
@@ -33,6 +38,7 @@ void FPyReferenceCollector::RemoveWrappedInstance(FPyWrapperBase* InInstance)
 
 void FPyReferenceCollector::AddReferencedObjects(FReferenceCollector& InCollector)
 {
+	if (!PythonWrappedInstances.IsEmpty())
 	{
 		FPyScopedGIL GIL;
 		for (FPyWrapperBase* PythonWrappedInstance : PythonWrappedInstances)
@@ -81,6 +87,28 @@ void FPyReferenceCollector::PurgeUnrealObjectReferences(const TArrayView<const U
 
 void FPyReferenceCollector::PurgeUnrealGeneratedTypes()
 {
+#if WITH_EDITOR
+	// Only do the reference cleaning when the editor is not shuthing down to avoid wasting time.
+	if (!IsEngineExitRequested())
+	{
+		// Preclean to only grab the object that are refered by other systems
+		CollectGarbage(GARBAGE_COLLECTION_KEEPFLAGS);
+
+		TArray<UObject*> PythonObjectInstances;
+		ForEachObjectOfClass(UPythonGeneratedClass::StaticClass(), [&PythonObjectInstances](UObject* InObject)
+			{
+				UPythonGeneratedClass* PythonGeneratedClass = CastChecked<UPythonGeneratedClass>(InObject);
+				ForEachObjectOfClass(PythonGeneratedClass, [&PythonObjectInstances](UObject* InInnerObject)
+					{
+						PythonObjectInstances.Add(InInnerObject);
+					}, false, RF_NoFlags, EInternalObjectFlags::Native);
+			}, false, RF_ClassDefaultObject, EInternalObjectFlags::Native);
+
+		ObjectTools::ForceReplaceReferences(nullptr, PythonObjectInstances);
+		CollectGarbage(GARBAGE_COLLECTION_KEEPFLAGS);
+	}
+#endif //WITH_EDITOR
+
 	FPurgingReferenceCollector PurgingReferenceCollector;
 	TArray<FWeakObjectPtr> WeakReferencesToPurgedObjects;
 
@@ -199,10 +227,8 @@ void FPyReferenceCollector::PurgeUnrealGeneratedTypes()
 
 	if (PurgingReferenceCollector.HasObjectToPurge())
 	{
-		Py_BEGIN_ALLOW_THREADS
 		AddReferencedObjects(PurgingReferenceCollector);
 		CollectGarbage(GARBAGE_COLLECTION_KEEPFLAGS);
-		Py_END_ALLOW_THREADS
 
 		for (const FWeakObjectPtr& WeakReferencesToPurgedObject : WeakReferencesToPurgedObjects)
 		{
@@ -222,6 +248,11 @@ void FPyReferenceCollector::PurgeUnrealGeneratedTypes()
 	}
 }
 
+void FPyReferenceCollector::Shutdown()
+{
+	PythonWrappedInstances.Reset();
+}
+
 void FPyReferenceCollector::AddReferencedObjectsFromDelegate(FReferenceCollector& InCollector, FScriptDelegate& InDelegate)
 {
 	// Keep the delegate object alive if it's using a Python proxy instance
diff --git a/Engine/Plugins/Experimental/PythonScriptPlugin/Source/PythonScriptPlugin/Private/PyReferenceCollector.h b/Engine/Plugins/Experimental/PythonScriptPlugin/Source/PythonScriptPlugin/Private/PyReferenceCollector.h
--- a/Engine/Plugins/Experimental/PythonScriptPlugin/Source/PythonScriptPlugin/Private/PyReferenceCollector.h
+++ b/Engine/Plugins/Experimental/PythonScriptPlugin/Source/PythonScriptPlugin/Private/PyReferenceCollector.h
@@ -48,6 +48,8 @@ public:
 	/** Purge any Python references to any Unreal types generated from or for Python */
 	void PurgeUnrealGeneratedTypes();
 
+	void Shutdown();
+
 	//~ FGCObject interface
 	virtual void AddReferencedObjects(FReferenceCollector& InCollector) override;
 	virtual FString GetReferencerName() const override;
diff --git a/Engine/Plugins/Experimental/PythonScriptPlugin/Source/PythonScriptPlugin/Private/PyWrapperTypeRegistry.cpp b/Engine/Plugins/Experimental/PythonScriptPlugin/Source/PythonScriptPlugin/Private/PyWrapperTypeRegistry.cpp
--- a/Engine/Plugins/Experimental/PythonScriptPlugin/Source/PythonScriptPlugin/Private/PyWrapperTypeRegistry.cpp
+++ b/Engine/Plugins/Experimental/PythonScriptPlugin/Source/PythonScriptPlugin/Private/PyWrapperTypeRegistry.cpp
@@ -832,6 +832,12 @@ void FPyWrapperTypeReinstancer::ProcessPending()
 	}
 }
 
+void FPyWrapperTypeReinstancer::Shutdown()
+{
+	ClassesToReinstance.Empty();
+	StructsToReinstance.Empty();
+}
+
 void FPyWrapperTypeReinstancer::AddReferencedObjects(FReferenceCollector& InCollector)
 {
 	for (auto& ClassToReinstancePair : ClassesToReinstance)
@@ -3296,6 +3302,25 @@ void FPyWrapperTypeRegistry::GenerateStubCodeForWrappedTypes(const EPyOnlineDocs
 	}
 }
 
+void FPyWrapperTypeRegistry::Shutdown()
+{
+	PythonWrappedTypeNameToUnrealTypeName.Empty();
+	PythonWrappedClasses.Empty();
+	PythonWrappedStructs.Empty();
+	PythonWrappedEnums.Empty();
+	PythonWrappedDelegates.Empty();
+	GeneratedWrappedTypes.Empty();
+	GeneratedWrappedTypesForModule.Empty();
+	UnresolvedMakeFuncs.Empty();
+	UnresolvedBreakFuncs.Empty();
+	OrphanedWrappedTypes.Empty();
+	NativePythonModules.Empty();
+	OnModuleDirtiedDelegate.Clear();
+	InlineStructFactories.Empty();
+
+	bCanRegisterInlineStructFactories.store(true);
+}
+
 void FPyWrapperTypeRegistry::GenerateStubCodeForWrappedType(PyTypeObject* PyType, const PyGenUtil::FGeneratedWrappedType* GeneratedTypeData, FPyFileWriter& OutPythonScript, FPyOnlineDocsSection* OutOnlineDocsSection)
 {
 	const FString PyTypeName = UTF8_TO_TCHAR(PyType->tp_name);
@@ -4124,6 +4149,23 @@ void FPyWrapperTypeRegistry::UnregisterPythonTypeName(const FString& InPythonTyp
 	);
 }
 
+void ShutdownPyWrappers()
+{
+	FPyWrapperObjectFactory::Get().Shutdown();
+	FPyWrapperStructFactory::Get().Shutdown();
+	FPyWrapperDelegateFactory::Get().Shutdown();
+	FPyWrapperMulticastDelegateFactory::Get().Shutdown();
+	FPyWrapperNameFactory::Get().Shutdown();
+	FPyWrapperTextFactory::Get().Shutdown();
+	FPyWrapperArrayFactory::Get().Shutdown();
+	FPyWrapperFixedArrayFactory::Get().Shutdown();
+	FPyWrapperSetFactory::Get().Shutdown();
+	FPyWrapperMapFactory::Get().Shutdown();
+	FPyWrapperFieldPathFactory::Get().Shutdown();
+	FPyWrapperTypeReinstancer::Get().Shutdown();
+	FPyWrapperTypeRegistry::Get().Shutdown();
+}
+
 #endif	// WITH_PYTHON
 
 #undef LOCTEXT_NAMESPACE
diff --git a/Engine/Plugins/Experimental/PythonScriptPlugin/Source/PythonScriptPlugin/Private/PyWrapperTypeRegistry.h b/Engine/Plugins/Experimental/PythonScriptPlugin/Source/PythonScriptPlugin/Private/PyWrapperTypeRegistry.h
--- a/Engine/Plugins/Experimental/PythonScriptPlugin/Source/PythonScriptPlugin/Private/PyWrapperTypeRegistry.h
+++ b/Engine/Plugins/Experimental/PythonScriptPlugin/Source/PythonScriptPlugin/Private/PyWrapperTypeRegistry.h
@@ -123,6 +123,11 @@ public:
 		MappedInstances.Remove(FInternalKey(TPyWrapperTypeFactoryConversion<UnrealType, KeyType>::UnrealTypeToKeyType(InUnrealInstance), InWrappedPyType));
 	}
 
+	void Shutdown()
+	{
+		MappedInstances.Reset();
+	}
+
 protected:
 	/** Callback used to initialize this type */
 	typedef TFunctionRef<int(PythonType*)> FCreateInstanceInitializerFunc;
@@ -339,6 +344,8 @@ public:
 	/** Process any pending re-instance requests */
 	void ProcessPending();
 
+	void Shutdown();
+
 private:
 	/** Pending pairs of classes that to be re-instanced */
 	TArray<TPair<TObjectPtr<UPythonGeneratedClass>, TObjectPtr<UPythonGeneratedClass>>> ClassesToReinstance;
@@ -488,6 +495,8 @@ public:
 	/** Generate stub Python code for our wrapped types */
 	void GenerateStubCodeForWrappedTypes(const EPyOnlineDocsFilterFlags InDocGenFlags = EPyOnlineDocsFilterFlags::IncludeNone) const;
 
+	void Shutdown();
+
 private:
 	/** Call the game-thread only logic */
 	void ExecutePostInitOnGameThread();
@@ -553,4 +562,6 @@ private:
 	FOnModuleDirtied OnModuleDirtiedDelegate;
 };
 
+void ShutdownPyWrappers();
+
 #endif	// WITH_PYTHON
diff --git a/Engine/Plugins/Experimental/PythonScriptPlugin/Source/PythonScriptPlugin/Private/PythonScriptPlugin.cpp b/Engine/Plugins/Experimental/PythonScriptPlugin/Source/PythonScriptPlugin/Private/PythonScriptPlugin.cpp
--- a/Engine/Plugins/Experimental/PythonScriptPlugin/Source/PythonScriptPlugin/Private/PythonScriptPlugin.cpp
+++ b/Engine/Plugins/Experimental/PythonScriptPlugin/Source/PythonScriptPlugin/Private/PythonScriptPlugin.cpp
@@ -757,6 +757,26 @@ bool FPythonScriptPlugin::ForceEnablePythonAtRuntime(const FString& Location)
 #endif
 }
 
+void FPythonScriptPlugin::ForceDisablePythonAtRuntime(UE::FSourceLocation Location)
+{
+#if WITH_PYTHON
+	if (bIsInterpreterInitialized)
+	{
+		const FString CallerLocation = UE::SourceLocation::Full(Location).ToString();
+		if (!UE::Python::Private::bIsEnabledByDefault || !UE::Python::Private::bIsPythonLifeCycleDynamic)
+		{
+			UE_LOG(LogPython, Warning, TEXT("Python disabled via IPythonScriptPlugin::DisablePythonAtRuntime: %s"), *CallerLocation);
+		}
+		else
+		{
+			UE_LOG(LogPython, Log, TEXT("Python disabled via IPythonScriptPlugin::DisablePythonAtRuntime: %s"), *CallerLocation);
+		}
+
+		ShutdownPython();
+		bIsForceEnabledAtRuntime = false;
+	}
+#endif
+}
 
 bool FPythonScriptPlugin::ExecPythonCommand(const TCHAR* InPythonCommand)
 {
@@ -1006,8 +1026,7 @@ void FPythonScriptPlugin::ConfigureAndInitializePython()
 	FPythonScriptPluginEditorStyle::Get();
 	if (UToolMenus::IsToolMenuUIEnabled())
 	{
-		check(CmdMenu == nullptr);
-		CmdMenu = new FPythonCommandMenuImpl();
+		CmdMenu = MakeUnique<FPythonCommandMenuImpl>();
 		CmdMenu->OnStartupMenu();
 
 		UToolMenus::Get()->RegisterStringCommandHandler("Python", FToolMenuExecuteString::CreateLambda([this](const FString& InString, const FToolMenuContext& InContext) {
@@ -1023,30 +1042,10 @@ void FPythonScriptPlugin::ShutdownModule()
 {
 #if WITH_EDITOR
 	FEditorPythonExecuter::OnShutdownModule();
+	FCoreDelegates::OnPostEngineInit.RemoveAll(this);
 #endif
 
 #if WITH_PYTHON
-	FCoreDelegates::OnPreExit.RemoveAll(this);
-	FCoreDelegates::OnPostEngineInit.RemoveAll(this);
-
-	RemoteExecution.Reset();
-
-#if WITH_EDITOR
-	if (CmdMenu)
-	{
-		CmdMenu->OnShutdownMenu();
-		delete CmdMenu;
-		CmdMenu = nullptr;
-	}
-
-	if (UToolMenus* ToolMenus = UToolMenus::TryGet())
-	{
-		ToolMenus->UnregisterStringCommandHandler("Python");
-	}
-#endif // WITH_EDITOR
-
-	IModularFeatures::Get().UnregisterModularFeature(IConsoleCommandExecutor::ModularFeatureName(), &CmdExec);
-	IModularFeatures::Get().UnregisterModularFeature(IConsoleCommandExecutor::ModularFeatureName(), &CmdREPLExec);
 	ShutdownPython();
 #endif	// WITH_PYTHON
 }
@@ -1434,6 +1433,27 @@ void FPythonScriptPlugin::ShutdownPython()
 		return;
 	}
 
+	FCoreDelegates::OnPreExit.RemoveAll(this);
+
+	RemoteExecution.Reset();
+
+#if WITH_EDITOR
+	if (CmdMenu)
+	{
+		CmdMenu->OnShutdownMenu();
+		CmdMenu.Reset();
+	}
+
+	if (UToolMenus* ToolMenus = UToolMenus::TryGet())
+	{
+		ToolMenus->UnregisterStringCommandHandler("Python");
+	}
+#endif // WITH_EDITOR
+
+	IModularFeatures::Get().UnregisterModularFeature(IConsoleCommandExecutor::ModularFeatureName(), &CmdExec);
+	IModularFeatures::Get().UnregisterModularFeature(IConsoleCommandExecutor::ModularFeatureName(), &CmdREPLExec);
+
+
 	if (PyMainThreadState)
 	{
 		// We need to restore the original GIL prior to calling Py_Finalize
@@ -1481,8 +1501,6 @@ void FPythonScriptPlugin::ShutdownPython()
 	FEditorSupportDelegates::PrepareToCleanseEditorObject.RemoveAll(this);
 #endif	// WITH_EDITOR
 
-	FPyReferenceCollector::Get().PurgeUnrealGeneratedTypes();
-
 #if WITH_EDITOR
 	PyEditor::ShutdownModule();
 #endif	// WITH_EDITOR
@@ -1503,10 +1521,12 @@ void FPythonScriptPlugin::ShutdownPython()
 		Py_Finalize();
 	}
 
-	bIsConfigured = false;
+	ShutdownPyWrappers();
+	FPyReferenceCollector::Get().Shutdown();
+	FPyReferenceCollector::Get().PurgeUnrealGeneratedTypes();
+
 	bIsInterpreterInitialized = false;
 	bIsFullyInitialized = false;
-	bIsEnabled = false;
 	bIsForceEnabledAtRuntime = false;
 	bHasTicked = false;
 }
@@ -1560,7 +1580,7 @@ void FPythonScriptPlugin::GenerateStubCode()
 
 void FPythonScriptPlugin::RunStartupScripts()
 {
-	if (bIsFullyInitialized)
+	if (bIsFullyInitialized || !bIsInterpreterInitialized)
 	{
 		return;
 	}
@@ -1574,9 +1594,11 @@ void FPythonScriptPlugin::RunStartupScripts()
 		PySysPaths = PyUtil::GetSystemPaths();
 	}
 
-	auto RunTimedStartupScript = [this](const FString& StartupScript, const FText& StartupScriptInfoText, const bool bRunAsFile)
+	double InitTotalTime = 0.0;
+
+	auto RunTimedStartupScript = [this, &InitTotalTime](const FString& StartupScript, const FText& StartupScriptInfoText, const bool bRunAsFile)
 	{
-		UE_SCOPED_TIMER(*StartupScriptInfoText.ToString(), LogPython, Display);
+		FScopedDurationTimeCustomLogger TimeLogger(*StartupScriptInfoText.ToString(), InitTotalTime, [](const TCHAR* Msg) { UE_LOG(LogPython, Display, TEXT("%s"), Msg) });
 		if (bRunAsFile)
 		{
 			// Execute these files in the "public" scope, as if their contents had been run directly in the console
diff --git a/Engine/Plugins/Experimental/PythonScriptPlugin/Source/PythonScriptPlugin/Private/PythonScriptPlugin.h b/Engine/Plugins/Experimental/PythonScriptPlugin/Source/PythonScriptPlugin/Private/PythonScriptPlugin.h
--- a/Engine/Plugins/Experimental/PythonScriptPlugin/Source/PythonScriptPlugin/Private/PythonScriptPlugin.h
+++ b/Engine/Plugins/Experimental/PythonScriptPlugin/Source/PythonScriptPlugin/Private/PythonScriptPlugin.h
@@ -109,7 +109,7 @@ public:
 	virtual bool IsPythonConfigured() const override;
 	virtual bool IsPythonInitialized() const override;
 	virtual bool ForceEnablePythonAtRuntime(UE::FSourceLocation Location = UE::FSourceLocation::Current()) override;
-	bool ForceEnablePythonAtRuntime(const FString& Location);
+	virtual void ForceDisablePythonAtRuntime(UE::FSourceLocation Location = UE::FSourceLocation::Current()) override;
 	virtual bool ExecPythonCommand(const TCHAR* InPythonCommand) override;
 	virtual bool ExecPythonCommandEx(FPythonCommandEx& InOutPythonCommand) override;
 	virtual FString GetInterpreterExecutablePath() const override;
@@ -117,6 +117,8 @@ public:
 	virtual FSimpleMulticastDelegate& OnPythonInitialized() override;
 	virtual FSimpleMulticastDelegate& OnPythonShutdown() override;
 
+	bool ForceEnablePythonAtRuntime(const FString& Location);
+
 	//~ IModuleInterface interface
 	virtual void StartupModule() override;
 	virtual void ShutdownModule() override;
@@ -167,10 +169,6 @@ private:
 
 	void ShutdownPython();
 
-	void InitPipInstaller();
-
-	void RunPipInstaller();
-
 	void RequestStubCodeGeneration();
 
 	void GenerateStubCode();
@@ -214,7 +212,7 @@ private:
 	TUniquePtr<FPythonScriptRemoteExecution> RemoteExecution;
 	FPythonCommandExecutor CmdExec;
 	FPythonREPLCommandExecutor CmdREPLExec;
-	IPythonCommandMenu* CmdMenu = nullptr;
+	TUniquePtr<IPythonCommandMenu> CmdMenu = nullptr;
 	FTSTicker::FDelegateHandle TickHandle;
 	FTSTicker::FDelegateHandle TickOnceHandle;
 	FTSTicker::FDelegateHandle ModuleDelayedHandle;
diff --git a/Engine/Plugins/Experimental/PythonScriptPlugin/Source/PythonScriptPlugin/Public/IPythonScriptPlugin.h b/Engine/Plugins/Experimental/PythonScriptPlugin/Source/PythonScriptPlugin/Public/IPythonScriptPlugin.h
--- a/Engine/Plugins/Experimental/PythonScriptPlugin/Source/PythonScriptPlugin/Public/IPythonScriptPlugin.h
+++ b/Engine/Plugins/Experimental/PythonScriptPlugin/Source/PythonScriptPlugin/Public/IPythonScriptPlugin.h
@@ -37,10 +37,15 @@ public:
 
 	/**
 	 * Force Python to be enabled and initialized, regardless of the settings that control its default enabled state.
-	 * @return True if Python was requested to be enabled. Use IsPythonInitialized to verify that it actually initialized.
+	 * @return False if Python can't be enabled ever. Use IsPythonInitialized to verify that it actually initialized.
 	 */
 	virtual bool ForceEnablePythonAtRuntime(UE::FSourceLocation Location = UE::FSourceLocation::Current()) = 0;
 
+	/**
+	 * Force Python to shutdown.
+	 */
+	virtual void ForceDisablePythonAtRuntime(UE::FSourceLocation Location = UE::FSourceLocation::Current()) = 0;
+
 	/**
 	 * Execute the given Python command.
 	 * This may be literal Python code, or a file (with optional arguments) that you want to run.
