Add to_dict to Python bindings so that you can access broken structures by their property names.

#Synced-CL 49693328

diff --git a/Engine/Plugins/Experimental/PythonScriptPlugin/Source/PythonScriptPlugin/Private/PyWrapperStruct.cpp b/Engine/Plugins/Experimental/PythonScriptPlugin/Source/PythonScriptPlugin/Private/PyWrapperStruct.cpp
--- a/Engine/Plugins/Experimental/PythonScriptPlugin/Source/PythonScriptPlugin/Private/PyWrapperStruct.cpp
+++ b/Engine/Plugins/Experimental/PythonScriptPlugin/Source/PythonScriptPlugin/Private/PyWrapperStruct.cpp
@@ -411,22 +411,22 @@ PyObject* FPyWrapperStruct::BreakStruct(FPyWrapperStruct* InSelf)
 		return CallBreakFunction_Impl(InSelf, StructMetaData->BreakFunc);
 	}
 
-	// Generic implementation just creates a tuple from each property
-	FPyObjectPtr PyPropTuple = FPyObjectPtr::StealReference(PyTuple_New(StructMetaData->InitParams.Num()));
+	// Generic implementation just creates a dict with each property
+	FPyObjectPtr PyPropDict = FPyObjectPtr::StealReference(PyDict_New());
 	for (int32 ParamIndex = 0; ParamIndex < StructMetaData->InitParams.Num(); ++ParamIndex)
 	{
 		const PyGenUtil::FGeneratedWrappedMethodParameter& InitParam = StructMetaData->InitParams[ParamIndex];
-
 		PyObject* PyValue = nullptr;
 		if (!PyConversion::PythonizeProperty_InContainer(InitParam.ParamProp, InSelf->StructInstance, 0, PyValue))
 		{
 			PyUtil::SetPythonError(PyExc_TypeError, InSelf, *FString::Printf(TEXT("Failed to convert property '%s' (%s) for attribute '%s' on '%s'"), *InitParam.ParamProp->GetName(), *InitParam.ParamProp->GetClass()->GetName(), UTF8_TO_TCHAR(InitParam.ParamName.GetData()), *InSelf->ScriptStruct->GetName()));
 			return nullptr;
 		}
-		PyTuple_SetItem(PyPropTuple, ParamIndex, PyValue); // SetItem steals the reference
+		FPyObjectPtr Key = FPyObjectPtr::StealReference(PyUnicode_FromString(InitParam.ParamName.GetData()));
+		PyDict_SetItem(PyPropDict, Key, PyValue);
 	}
 
-	return PyPropTuple.Release();
+	return PyPropDict.Release();
 }
 
 PyObject* FPyWrapperStruct::GetPropertyValue(FPyWrapperStruct* InSelf, const PyGenUtil::FGeneratedWrappedProperty& InPropDef, const char* InPythonAttrName)
@@ -522,7 +522,7 @@ PyObject* FPyWrapperStruct::CallBreakFunction_Impl(FPyWrapperStruct* InSelf, con
 		{
 			return nullptr;
 		}
-		FPyObjectPtr PyPropTuple = FPyObjectPtr::StealReference(PyTuple_New(InFuncDef.OutputParams.Num()));
+		FPyObjectPtr PyPropDict = FPyObjectPtr::StealReference(PyDict_New());
 		for (int32 ParamIndex = 0; ParamIndex < InFuncDef.OutputParams.Num(); ++ParamIndex)
 		{
 			const PyGenUtil::FGeneratedWrappedMethodParameter& ParamDef = InFuncDef.OutputParams[ParamIndex];
@@ -533,9 +533,10 @@ PyObject* FPyWrapperStruct::CallBreakFunction_Impl(FPyWrapperStruct* InSelf, con
 				PyUtil::SetPythonError(PyExc_TypeError, InSelf, *FString::Printf(TEXT("Failed to convert return property '%s' when calling function '%s.%s' on '%s'"), UTF8_TO_TCHAR(ParamDef.ParamName.GetData()), *Class->GetName(), *InFuncDef.Func->GetName(), *Obj->GetName()));
 				return nullptr;
 			}
-			PyTuple_SetItem(PyPropTuple, ParamIndex, PyValue); // SetItem steals the reference
+			FPyObjectPtr Key = FPyObjectPtr::StealReference(PyUnicode_FromString(ParamDef.ParamName.GetData()));
+			PyDict_SetItem(PyPropDict, Key, PyValue); // SetItem steals the reference
 		}
-		return PyPropTuple.Release();
+		return PyPropDict.Release();
 	}
 
 	Py_RETURN_NONE;
@@ -993,6 +994,26 @@ PyTypeObject InitializePyWrapperStructType()
 		}
 
 		static PyObject* ToTuple(FPyWrapperStruct* InSelf)
+		{
+			FPyObjectPtr PyDict = FPyObjectPtr::StealReference(FPyWrapperStruct::BreakStruct(InSelf));
+			if (!PyDict || !PyDict_Check(PyDict))
+			{
+				PyErr_SetString(PyExc_TypeError, "Expected a dictionary object");
+				return nullptr;
+			}
+
+			// Get the "values" from the object
+			FPyObjectPtr Values = FPyObjectPtr::StealReference(PyDict_Values(PyDict));
+			if (!Values)
+			{
+				return nullptr; // Exception set internally
+			}
+
+			// Convert the values (which are iterable) into a new tuple
+			return PyList_AsTuple(Values);
+		}
+
+		static PyObject* ToDict(FPyWrapperStruct* InSelf)
 		{
 			return FPyWrapperStruct::BreakStruct(InSelf);
 		}
@@ -1276,6 +1297,7 @@ PyTypeObject InitializePyWrapperStructType()
 		{ "copy", PyCFunctionCast(&FMethods::Copy), METH_NOARGS, "copy(self) -> Any -- copy this Unreal struct" },
 		{ "assign", PyCFunctionCast(&FMethods::Assign), METH_VARARGS, "assign(self, other: object) -> None -- assign the value of this Unreal struct to value of the given object" },
 		{ "to_tuple", PyCFunctionCast(&FMethods::ToTuple), METH_NOARGS, "to_tuple(self) -> Tuple[object, ...] -- break this Unreal struct into a tuple of its properties" },
+		{ "to_dict", PyCFunctionCast(&FMethods::ToDict), METH_NOARGS, "to_dict(self) -> Dict[str, object] -- convert this Unreal struct into a dict of its properties" },
 		{ "get_editor_property", PyCFunctionCast(&FMethods::GetEditorProperty), METH_VARARGS | METH_KEYWORDS, "get_editor_property(self, name: Union[Name, str]) -> object -- get the value of any property visible to the editor" },
 		{ "set_editor_property", PyCFunctionCast(&FMethods::SetEditorProperty), METH_VARARGS | METH_KEYWORDS, "set_editor_property(self, name: Union[Name, str], value: object, notify_mode: PropertyAccessChangeNotifyMode=PropertyAccessChangeNotifyMode.DEFAULT) -> None -- set the value of any property visible to the editor, ensuring that the pre/post change notifications are called" },
 		{ "set_editor_properties", PyCFunctionCast(&FMethods::SetEditorProperties), METH_VARARGS, "set_editor_properties(self, properties: Mapping[str, object]) -> None -- set the value of any properties visible to the editor (from a name->value dict), ensuring that the pre/post change notifications are called" },
