Fix leaked value reference in when splitting struct in Python.

#Synced-CL 49756264

diff --git a/Engine/Plugins/Experimental/PythonScriptPlugin/Source/PythonScriptPlugin/Private/PyWrapperStruct.cpp b/Engine/Plugins/Experimental/PythonScriptPlugin/Source/PythonScriptPlugin/Private/PyWrapperStruct.cpp
--- a/Engine/Plugins/Experimental/PythonScriptPlugin/Source/PythonScriptPlugin/Private/PyWrapperStruct.cpp
+++ b/Engine/Plugins/Experimental/PythonScriptPlugin/Source/PythonScriptPlugin/Private/PyWrapperStruct.cpp
@@ -422,8 +422,9 @@ PyObject* FPyWrapperStruct::BreakStruct(FPyWrapperStruct* InSelf)
 			PyUtil::SetPythonError(PyExc_TypeError, InSelf, *FString::Printf(TEXT("Failed to convert property '%s' (%s) for attribute '%s' on '%s'"), *InitParam.ParamProp->GetName(), *InitParam.ParamProp->GetClass()->GetName(), UTF8_TO_TCHAR(InitParam.ParamName.GetData()), *InSelf->ScriptStruct->GetName()));
 			return nullptr;
 		}
+		FPyObjectPtr Value = FPyObjectPtr::StealReference(PyValue);
 		FPyObjectPtr Key = FPyObjectPtr::StealReference(PyUnicode_FromString(InitParam.ParamName.GetData()));
-		PyDict_SetItem(PyPropDict, Key, PyValue);
+		PyDict_SetItem(PyPropDict, Key, Value);
 	}
 
 	return PyPropDict.Release();
@@ -534,7 +535,8 @@ PyObject* FPyWrapperStruct::CallBreakFunction_Impl(FPyWrapperStruct* InSelf, con
 				return nullptr;
 			}
 			FPyObjectPtr Key = FPyObjectPtr::StealReference(PyUnicode_FromString(ParamDef.ParamName.GetData()));
-			PyDict_SetItem(PyPropDict, Key, PyValue); // SetItem steals the reference
+			FPyObjectPtr Value = FPyObjectPtr::StealReference(PyValue);
+			PyDict_SetItem(PyPropDict, Key, Value);
 		}
 		return PyPropDict.Release();
 	}
