Added Optional property support to Python

Python can now read and write Optional properties from UObjects and UStructs, where None is treated as unset.
Note: This does mean that you can't distinguish between an unset Optional and an Optional set to a null UObject, but the latter is an unusual use of Optional.

Python generated types can create Optional properties via a tuple of two values, where one is None and the other is the value type of the Optional (in either order), eg)
```
class TestOptional(unreal.Object):
prop = unreal.uproperty((str, None))

def func(self, opt):
pass
```

#Synced-CL 49717872

diff --git a/Engine/Plugins/Experimental/PythonScriptPlugin/Source/PythonScriptPlugin/Private/PyConversion.cpp b/Engine/Plugins/Experimental/PythonScriptPlugin/Source/PythonScriptPlugin/Private/PyConversion.cpp
--- a/Engine/Plugins/Experimental/PythonScriptPlugin/Source/PythonScriptPlugin/Private/PyConversion.cpp
+++ b/Engine/Plugins/Experimental/PythonScriptPlugin/Source/PythonScriptPlugin/Private/PyConversion.cpp
@@ -22,6 +22,7 @@
 #include "UObject/UnrealType.h"
 #include "UObject/EnumProperty.h"
 #include "UObject/TextProperty.h"
+#include "UObject/PropertyOptional.h"
 #include "UObject/PropertyPortFlags.h"
 
 #if WITH_PYTHON
@@ -1031,6 +1032,44 @@ FPyConversionResult NativizeProperty_Direct(PyObject* PyObj, const FProperty* Pr
 		PYCONVERSION_PROPERTY_RETURN(Result);
 	}
 
+	if (const FOptionalProperty* CastProp = CastField<FOptionalProperty>(Prop))
+	{
+		FPyConversionResult Result = FPyConversionResult::Failure();
+
+		if (PyObj == Py_None)
+		{
+			EmitPropertyChangeNotifications(InChangeNotify, !CastProp->IsSet(ValueAddr), [&]()
+			{
+				for (void* ArchInstValueAddr : InArchetypeInstValueAddrs)
+				{
+					CastProp->MarkUnset(ArchInstValueAddr);
+				}
+				CastProp->MarkUnset(ValueAddr);
+			});
+
+			Result = FPyConversionResult::Success();
+		}
+		else
+		{
+			PyUtil::FPropValueOnScope NewValue(PyUtil::FConstPropOnScope::ExternalReference(CastProp));
+			Result = NativizeProperty_Direct(PyObj, CastProp->GetValueProperty(), CastProp->MarkSetAndGetInitializedValuePointerToReplace(NewValue.GetValue()), {}, nullptr, SetErrorState);
+
+			if (Result)
+			{
+				EmitPropertyChangeNotifications(InChangeNotify, CastProp->Identical(ValueAddr, NewValue.GetValue(), PPF_None), [&]()
+				{
+					for (void* ArchInstValueAddr : InArchetypeInstValueAddrs)
+					{
+						CastProp->CopySingleValue(ArchInstValueAddr, NewValue.GetValue());
+					}
+					CastProp->CopySingleValue(ValueAddr, NewValue.GetValue());
+				});
+			}
+		}
+
+		PYCONVERSION_PROPERTY_RETURN(Result);
+	}
+
 #undef NATIVIZE_SETTER_PROPERTY
 #undef NATIVIZE_INLINE_PROPERTY
 
@@ -1167,6 +1206,20 @@ FPyConversionResult PythonizeProperty_Direct(const FProperty* Prop, const void*
 		return FPyConversionResult::Success();
 	}
 
+	if (const FOptionalProperty* CastProp = CastField<FOptionalProperty>(Prop))
+	{
+		if (CastProp->IsSet(ValueAddr))
+		{
+			PYCONVERSION_PROPERTY_RETURN(PythonizeProperty_Direct(CastProp->GetValueProperty(), CastProp->GetValuePointerForRead(ValueAddr), OutPyObj, ConversionMethod, OwnerPyObj, SetErrorState));
+		}
+		else
+		{
+			Py_INCREF(Py_None);
+			OutPyObj = Py_None;
+			return FPyConversionResult::Success();
+		}
+	}
+
 #undef PYTHONIZE_GETTER_PROPERTY
 
 	PYCONVERSION_RETURN(FPyConversionResult::Failure(), TEXT("PythonizeProperty"), *FString::Printf(TEXT("Cannot pythonize '%s' (%s). %s conversion not implemented!"), *Prop->GetName(), *Prop->GetClass()->GetName(), *Prop->GetClass()->GetName()));
diff --git a/Engine/Plugins/Experimental/PythonScriptPlugin/Source/PythonScriptPlugin/Private/PyGenUtil.cpp b/Engine/Plugins/Experimental/PythonScriptPlugin/Source/PythonScriptPlugin/Private/PyGenUtil.cpp
--- a/Engine/Plugins/Experimental/PythonScriptPlugin/Source/PythonScriptPlugin/Private/PyGenUtil.cpp
+++ b/Engine/Plugins/Experimental/PythonScriptPlugin/Source/PythonScriptPlugin/Private/PyGenUtil.cpp
@@ -21,6 +21,7 @@
 #include "UObject/NameTypes.h"
 #include "UObject/TextProperty.h"
 #include "UObject/CoreRedirects.h"
+#include "UObject/PropertyOptional.h"
 #include "UObject/PropertyPortFlags.h"
 #include "UObject/UnrealType.h"
 #include "Async/EventCount.h"
@@ -2398,6 +2399,10 @@ void PythonizeValueImpl(const FProperty* InProp, const void* InPropValue, const
 		{
 			OutPythonDefaultValue += TEXT("None");
 		}
+		else if (const FOptionalProperty* OptionalProp = CastField<const FOptionalProperty>(InProp))
+		{
+			OutPythonDefaultValue += TEXT("None");
+		}
 		else if (const FStructProperty* StructProp = CastField<const FStructProperty>(InProp))
 		{
 			PythonizeStructValueImpl(StructProp->Struct, PropArrValue, InFlags, OutPythonDefaultValue);
@@ -3399,6 +3404,10 @@ FString GetPropertyTypePythonName(const FProperty* InProp, const bool InIncludeU
 			? FString::Printf(TEXT("Optional[%s%s]"), UnrealNamespace, *GetClassPythonName(InterfaceProp->InterfaceClass)) // None can be assigned to Interface.
 			: FString::Printf(TEXT("%s%s"), UnrealNamespace, *GetClassPythonName(InterfaceProp->InterfaceClass));
 	}
+	if (const FOptionalProperty* OptionalProperty = CastField<const FOptionalProperty>(InProp))
+	{
+		return FString::Printf(TEXT("Optional[%s]"), *GetPropertyTypePythonName(OptionalProperty->GetValueProperty(), InIncludeUnrealNamespace, InIsForDocString, InPythonizeFlags));
+	}
 	if (const FStructProperty* StructProp = CastField<const FStructProperty>(InProp))
 	{
 		// If type hinting is enabled and the context allows type coercion (for an input parameters), the implementation support several types in place of the struct.
diff --git a/Engine/Plugins/Experimental/PythonScriptPlugin/Source/PythonScriptPlugin/Private/PyReferenceCollector.cpp b/Engine/Plugins/Experimental/PythonScriptPlugin/Source/PythonScriptPlugin/Private/PyReferenceCollector.cpp
--- a/Engine/Plugins/Experimental/PythonScriptPlugin/Source/PythonScriptPlugin/Private/PyReferenceCollector.cpp
+++ b/Engine/Plugins/Experimental/PythonScriptPlugin/Source/PythonScriptPlugin/Private/PyReferenceCollector.cpp
@@ -11,6 +11,7 @@
 #include "UObject/ObjectRedirector.h"
 #include "UObject/UnrealType.h"
 #include "UObject/UObjectHash.h"
+#include "UObject/PropertyOptional.h"
 #include "UObject/PurgingReferenceCollector.h"
 
 #if WITH_EDITOR
@@ -327,22 +328,22 @@ void FPyReferenceCollector::AddReferencedObjectsFromPropertyInternal(FReferenceC
 	{
 		if (EnumHasAnyFlags(InFlags, EPyReferenceCollectorFlags::IncludeInterfaces))
 		{
-		  for (int32 ArrIndex = 0; ArrIndex < InProp->ArrayDim; ++ArrIndex)
-		  {
-			  void* ValuePtr = CastProp->ContainerPtrToValuePtr<void>(InBaseAddr, ArrIndex);
-			  UObject* CurObjVal = CastProp->GetPropertyValue(ValuePtr).GetObject();
-			  if (CurObjVal)
-			  {
-				  TObjectPtr<UObject> NewObjVal = CurObjVal;
-          InCollector.AddReferencedObject(NewObjVal);
-  
-				  if (NewObjVal != CurObjVal)
-				  {
-					  OutValueChanged = true;
-					  CastProp->SetPropertyValue(ValuePtr, FScriptInterface(NewObjVal, NewObjVal ? NewObjVal->GetInterfaceAddress(CastProp->InterfaceClass) : nullptr));
-				  }
-			  }
-		  }
+			for (int32 ArrIndex = 0; ArrIndex < InProp->ArrayDim; ++ArrIndex)
+			{
+				void* ValuePtr = CastProp->ContainerPtrToValuePtr<void>(InBaseAddr, ArrIndex);
+				UObject* CurObjVal = CastProp->GetPropertyValue(ValuePtr).GetObject();
+				if (CurObjVal)
+				{
+					TObjectPtr<UObject> NewObjVal = CurObjVal;
+					InCollector.AddReferencedObject(NewObjVal);
+
+					if (NewObjVal != CurObjVal)
+					{
+						OutValueChanged = true;
+						CastProp->SetPropertyValue(ValuePtr, FScriptInterface(NewObjVal, NewObjVal ? NewObjVal->GetInterfaceAddress(CastProp->InterfaceClass) : nullptr));
+					}
+				}
+			}
 		}
 		return;
 	}
@@ -363,8 +364,8 @@ void FPyReferenceCollector::AddReferencedObjectsFromPropertyInternal(FReferenceC
 	{
 		if (EnumHasAnyFlags(InFlags, EPyReferenceCollectorFlags::IncludeDelegates))
 		{
-		for (int32 ArrIndex = 0; ArrIndex < InProp->ArrayDim; ++ArrIndex)
-		{
+			for (int32 ArrIndex = 0; ArrIndex < InProp->ArrayDim; ++ArrIndex)
+			{
 				FScriptDelegate* Value = CastProp->GetPropertyValuePtr(CastProp->ContainerPtrToValuePtr<void>(InBaseAddr, ArrIndex));
 				AddReferencedObjectsFromDelegate(InCollector, *Value);
 			}
@@ -391,13 +392,13 @@ void FPyReferenceCollector::AddReferencedObjectsFromPropertyInternal(FReferenceC
 	{
 		if (EnumHasAnyFlags(InFlags, EPyReferenceCollectorFlags::IncludeArrays))
 		{
-		for (int32 ArrIndex = 0; ArrIndex < InProp->ArrayDim; ++ArrIndex)
-		{
-			FScriptArrayHelper_InContainer ScriptArrayHelper(CastProp, InBaseAddr, ArrIndex);
-
-			const int32 ElementCount = ScriptArrayHelper.Num();
-			for (int32 ElementIndex = 0; ElementIndex < ElementCount; ++ElementIndex)
+			for (int32 ArrIndex = 0; ArrIndex < InProp->ArrayDim; ++ArrIndex)
 			{
+				FScriptArrayHelper_InContainer ScriptArrayHelper(CastProp, InBaseAddr, ArrIndex);
+
+				const int32 ElementCount = ScriptArrayHelper.Num();
+				for (int32 ElementIndex = 0; ElementIndex < ElementCount; ++ElementIndex)
+				{
 					AddReferencedObjectsFromPropertyInternal(InCollector, CastProp->Inner, ScriptArrayHelper.GetRawPtr(ElementIndex), InFlags, OutValueChanged);
 				}
 			}
@@ -409,22 +410,22 @@ void FPyReferenceCollector::AddReferencedObjectsFromPropertyInternal(FReferenceC
 	{
 		if (EnumHasAnyFlags(InFlags, EPyReferenceCollectorFlags::IncludeSets))
 		{
-		for (int32 ArrIndex = 0; ArrIndex < InProp->ArrayDim; ++ArrIndex)
-		{
-			bool bSetValuesChanged = false;
-			FScriptSetHelper_InContainer ScriptSetHelper(CastProp, InBaseAddr, ArrIndex);
-			for (FScriptSetHelper::FIterator It(ScriptSetHelper); It; ++It)
+			for (int32 ArrIndex = 0; ArrIndex < InProp->ArrayDim; ++ArrIndex)
 			{
-				AddReferencedObjectsFromPropertyInternal(InCollector, ScriptSetHelper.GetElementProperty(), ScriptSetHelper.GetElementPtr(It), InFlags, bSetValuesChanged);
-			}
+				bool bSetValuesChanged = false;
+				FScriptSetHelper_InContainer ScriptSetHelper(CastProp, InBaseAddr, ArrIndex);
+				for (FScriptSetHelper::FIterator It(ScriptSetHelper); It; ++It)
+				{
+					AddReferencedObjectsFromPropertyInternal(InCollector, ScriptSetHelper.GetElementProperty(), ScriptSetHelper.GetElementPtr(It), InFlags, bSetValuesChanged);
+				}
 
-			if (bSetValuesChanged)
-			{
-				OutValueChanged = true;
-				ScriptSetHelper.Rehash();
+				if (bSetValuesChanged)
+				{
+					OutValueChanged = true;
+					ScriptSetHelper.Rehash();
+				}
 			}
 		}
-		}
 		return;
 	}
 
@@ -432,28 +433,37 @@ void FPyReferenceCollector::AddReferencedObjectsFromPropertyInternal(FReferenceC
 	{
 		if (EnumHasAnyFlags(InFlags, EPyReferenceCollectorFlags::IncludeMaps))
 		{
-		for (int32 ArrIndex = 0; ArrIndex < InProp->ArrayDim; ++ArrIndex)
-		{
-			bool bMapKeysChanged = false;
-			bool bMapValuesChanged = false;
-			FScriptMapHelper_InContainer ScriptMapHelper(CastProp, InBaseAddr, ArrIndex);
-
-			for (FScriptMapHelper::FIterator It(ScriptMapHelper); It; ++It)
+			for (int32 ArrIndex = 0; ArrIndex < InProp->ArrayDim; ++ArrIndex)
 			{
+				bool bMapKeysChanged = false;
+				bool bMapValuesChanged = false;
+				FScriptMapHelper_InContainer ScriptMapHelper(CastProp, InBaseAddr, ArrIndex);
+
+				for (FScriptMapHelper::FIterator It(ScriptMapHelper); It; ++It)
+				{
 					// Note: We use the pair pointer below as AddReferencedObjectsFromPropertyInternal expects a base address and the key/value property will apply the correct offset from the base
 					AddReferencedObjectsFromPropertyInternal(InCollector, ScriptMapHelper.GetKeyProperty(), ScriptMapHelper.GetPairPtr(It), InFlags, bMapKeysChanged);
 					AddReferencedObjectsFromPropertyInternal(InCollector, ScriptMapHelper.GetValueProperty(), ScriptMapHelper.GetPairPtr(It), InFlags, bMapValuesChanged);
-			}
+				}
 
-			if (bMapKeysChanged || bMapValuesChanged)
-			{
-				OutValueChanged = true;
-				if (bMapKeysChanged)
+				if (bMapKeysChanged || bMapValuesChanged)
 				{
-					ScriptMapHelper.Rehash();
+					OutValueChanged = true;
+					if (bMapKeysChanged)
+					{
+						ScriptMapHelper.Rehash();
+					}
 				}
 			}
 		}
+		return;
+	}
+
+	if (const FOptionalProperty* CastProp = CastField<FOptionalProperty>(InProp))
+	{
+		if (CastProp->IsSet(InBaseAddr))
+		{
+			AddReferencedObjectsFromPropertyInternal(InCollector, CastProp->GetValueProperty(), CastProp->GetValuePointerForReadOrReplace(InBaseAddr), InFlags, OutValueChanged);
 		}
 		return;
 	}
diff --git a/Engine/Plugins/Experimental/PythonScriptPlugin/Source/PythonScriptPlugin/Private/PyTest.h b/Engine/Plugins/Experimental/PythonScriptPlugin/Source/PythonScriptPlugin/Private/PyTest.h
--- a/Engine/Plugins/Experimental/PythonScriptPlugin/Source/PythonScriptPlugin/Private/PyTest.h
+++ b/Engine/Plugins/Experimental/PythonScriptPlugin/Source/PythonScriptPlugin/Private/PyTest.h
@@ -91,6 +91,12 @@ public:
 
 	UPROPERTY(VisibleAnywhere, BlueprintReadOnly, Category = "Python|Internal")
 	mutable bool BoolMutable = false;
+
+	UPROPERTY(EditAnywhere, BlueprintReadWrite, Category = "Python|Internal")
+	TOptional<int32> OptionalInt;
+
+	UPROPERTY(EditAnywhere, BlueprintReadWrite, Category = "Python|Internal")
+	TOptional<FString> OptionalString;
 };
 
 /**
@@ -222,6 +228,12 @@ public:
 	UPROPERTY(EditDefaultsOnly, Category = "Python|Internal")
 	bool BoolDefaultsOnly;
 
+	UPROPERTY(EditAnywhere, BlueprintReadWrite, Category = "Python|Internal")
+	TOptional<int32> OptionalInt;
+
+	UPROPERTY(EditAnywhere, BlueprintReadWrite, Category = "Python|Internal")
+	TOptional<FString> OptionalString;
+
 	UFUNCTION(BlueprintImplementableEvent, Category = "Python|Internal")
 	int32 FuncBlueprintImplementable(const int32 InValue) const;
 
diff --git a/Engine/Plugins/Experimental/PythonScriptPlugin/Source/PythonScriptPlugin/Private/PyUtil.cpp b/Engine/Plugins/Experimental/PythonScriptPlugin/Source/PythonScriptPlugin/Private/PyUtil.cpp
--- a/Engine/Plugins/Experimental/PythonScriptPlugin/Source/PythonScriptPlugin/Private/PyUtil.cpp
+++ b/Engine/Plugins/Experimental/PythonScriptPlugin/Source/PythonScriptPlugin/Private/PyUtil.cpp
@@ -33,6 +33,7 @@
 #include "Templates/Casts.h"
 #include "UObject/EnumProperty.h"
 #include "UObject/Package.h"
+#include "UObject/PropertyOptional.h"
 #include "UObject/PropertyPortFlags.h"
 #include "UObject/StructOnScope.h"
 #include "UObject/TextProperty.h"
@@ -264,6 +265,11 @@ FPropertyDef::FPropertyDef(const FProperty* InProperty)
 		KeyDef = MakeShared<FPropertyDef>(MapProp->KeyProp);
 		ValueDef = MakeShared<FPropertyDef>(MapProp->ValueProp);
 	}
+
+	if (const FOptionalProperty* OptionalProp = CastField<FOptionalProperty>(InProperty))
+	{
+		ValueDef = MakeShared<FPropertyDef>(OptionalProp->GetValueProperty());
+	}
 }
 
 bool CalculatePropertyDef(PyTypeObject* InPyType, FPropertyDef& OutPropertyDef)
@@ -407,6 +413,44 @@ bool CalculatePropertyDef(PyObject* InPyObj, FPropertyDef& OutPropertyDef)
 		return true;
 	}
 
+	// A tuple of two values, where one is None, is considered an optional property
+	if (PyTuple_Check(InPyObj) && PyTuple_GET_SIZE(InPyObj) == 2)
+	{
+		PyObject* PyValueObj = nullptr;
+		int32 NoneCount = 0;
+
+		if (PyObject* PotentialPyValueObj = PyTuple_GET_ITEM(InPyObj, 0);
+			PotentialPyValueObj != Py_None)
+		{
+			PyValueObj = PotentialPyValueObj;
+		}
+		else
+		{
+			++NoneCount;
+		}
+
+		if (PyObject* PotentialPyValueObj = PyTuple_GET_ITEM(InPyObj, 1);
+			PotentialPyValueObj != Py_None)
+		{
+			PyValueObj = PotentialPyValueObj;
+		}
+		else
+		{
+			++NoneCount;
+		}
+		
+		if (PyValueObj && NoneCount == 1)
+		{
+			FPropertyDef ValuePropertyDef;
+			if (CalculatePropertyDef(PyValueObj, ValuePropertyDef))
+			{
+				OutPropertyDef.PropertyClass = FOptionalProperty::StaticClass();
+				OutPropertyDef.ValueDef = MakeShared<FPropertyDef>(MoveTemp(ValuePropertyDef));
+				return true;
+			}
+		}
+	}
+
 	return CalculatePropertyDef(PyType_Check(InPyObj) ? (PyTypeObject*)InPyObj : Py_TYPE(InPyObj), OutPropertyDef);
 }
 
@@ -490,6 +534,11 @@ FProperty* CreateProperty(const FPropertyDef& InPropertyDef, const int32 InArray
 			MapProp->ValueProp = CreateProperty(*InPropertyDef.ValueDef, 1, MapProp);
 		}
 
+		if (FOptionalProperty* OptionalProp = CastField<FOptionalProperty>(Prop))
+		{
+			OptionalProp->SetValueProperty(CreateProperty(*InPropertyDef.ValueDef, 1, OptionalProp));
+		}
+
 		// Need to manually call Link to fix-up some data (such as the C++ property flags) that are only set during Link
 		{
 			FArchive Ar;
diff --git a/Engine/Plugins/Experimental/PythonScriptPlugin/Source/PythonScriptPlugin/Private/PyUtil.h b/Engine/Plugins/Experimental/PythonScriptPlugin/Source/PythonScriptPlugin/Private/PyUtil.h
--- a/Engine/Plugins/Experimental/PythonScriptPlugin/Source/PythonScriptPlugin/Private/PyUtil.h
+++ b/Engine/Plugins/Experimental/PythonScriptPlugin/Source/PythonScriptPlugin/Private/PyUtil.h
@@ -384,7 +384,7 @@ namespace PyUtil
 		/** Key definition of this property (for map properties) */
 		TSharedPtr<FPropertyDef> KeyDef;
 
-		/** Value definition of this property (for array, set, and map properties) */
+		/** Value definition of this property (for array, set, map, and optional properties) */
 		TSharedPtr<FPropertyDef> ValueDef;
 	};
 
