Fixed Optional property support not working for Python-defined function return types

#rnx

#Synced-CL 50324810

diff --git a/Engine/Plugins/Experimental/PythonScriptPlugin/Source/PythonScriptPlugin/Private/PyUtil.cpp b/Engine/Plugins/Experimental/PythonScriptPlugin/Source/PythonScriptPlugin/Private/PyUtil.cpp
--- a/Engine/Plugins/Experimental/PythonScriptPlugin/Source/PythonScriptPlugin/Private/PyUtil.cpp
+++ b/Engine/Plugins/Experimental/PythonScriptPlugin/Source/PythonScriptPlugin/Private/PyUtil.cpp
@@ -413,6 +413,28 @@ bool CalculatePropertyDef(PyObject* InPyObj, FPropertyDef& OutPropertyDef)
 		return true;
 	}
 
+	if (PyObject* PyValueObj = nullptr;
+		IsOptionalPropertyDef(InPyObj, &PyValueObj))
+	{
+		FPropertyDef ValuePropertyDef;
+		if (CalculatePropertyDef(PyValueObj, ValuePropertyDef))
+		{
+			OutPropertyDef.PropertyClass = FOptionalProperty::StaticClass();
+			OutPropertyDef.ValueDef = MakeShared<FPropertyDef>(MoveTemp(ValuePropertyDef));
+			return true;
+		}
+	}
+
+	return CalculatePropertyDef(PyType_Check(InPyObj) ? (PyTypeObject*)InPyObj : Py_TYPE(InPyObj), OutPropertyDef);
+}
+
+bool IsOptionalPropertyDef(PyObject* InPyObj, PyObject** OutPyValueObj)
+{
+	if (OutPyValueObj)
+	{
+		*OutPyValueObj = nullptr;
+	}
+
 	// A tuple of two values, where one is None, is considered an optional property
 	if (PyTuple_Check(InPyObj) && PyTuple_GET_SIZE(InPyObj) == 2)
 	{
@@ -438,20 +460,18 @@ bool CalculatePropertyDef(PyObject* InPyObj, FPropertyDef& OutPropertyDef)
 		{
 			++NoneCount;
 		}
-		
+
 		if (PyValueObj && NoneCount == 1)
 		{
-			FPropertyDef ValuePropertyDef;
-			if (CalculatePropertyDef(PyValueObj, ValuePropertyDef))
+			if (OutPyValueObj)
 			{
-				OutPropertyDef.PropertyClass = FOptionalProperty::StaticClass();
-				OutPropertyDef.ValueDef = MakeShared<FPropertyDef>(MoveTemp(ValuePropertyDef));
-				return true;
+				*OutPyValueObj = PyValueObj;
 			}
+			return true;
 		}
 	}
 
-	return CalculatePropertyDef(PyType_Check(InPyObj) ? (PyTypeObject*)InPyObj : Py_TYPE(InPyObj), OutPropertyDef);
+	return false;
 }
 
 FProperty* CreateProperty(const FPropertyDef& InPropertyDef, const int32 InArrayDim, FFieldVariant InOuter, const FName InName)
diff --git a/Engine/Plugins/Experimental/PythonScriptPlugin/Source/PythonScriptPlugin/Private/PyUtil.h b/Engine/Plugins/Experimental/PythonScriptPlugin/Source/PythonScriptPlugin/Private/PyUtil.h
--- a/Engine/Plugins/Experimental/PythonScriptPlugin/Source/PythonScriptPlugin/Private/PyUtil.h
+++ b/Engine/Plugins/Experimental/PythonScriptPlugin/Source/PythonScriptPlugin/Private/PyUtil.h
@@ -394,6 +394,9 @@ namespace PyUtil
 	/** Given a Python instance, work out what kind of property we would need to create to hold this data */
 	bool CalculatePropertyDef(PyObject* InPyObj, FPropertyDef& OutPropertyDef);
 
+	/** Given a Python instance, work out whether this will produce an Optional property def, and if so (and requested), get the value that the Optional type will hold */
+	bool IsOptionalPropertyDef(PyObject* InPyObj, PyObject** OutPyValueObj = nullptr);
+
 	/** Given a property definition, create a property instance */
 	FProperty* CreateProperty(const FPropertyDef& InPropertyDef, const int32 InArrayDim = 1, FFieldVariant InOuter = nullptr, const FName InName = DefaultPythonPropertyName);
 
diff --git a/Engine/Plugins/Experimental/PythonScriptPlugin/Source/PythonScriptPlugin/Private/PyWrapperObject.cpp b/Engine/Plugins/Experimental/PythonScriptPlugin/Source/PythonScriptPlugin/Private/PyWrapperObject.cpp
--- a/Engine/Plugins/Experimental/PythonScriptPlugin/Source/PythonScriptPlugin/Private/PyWrapperObject.cpp
+++ b/Engine/Plugins/Experimental/PythonScriptPlugin/Source/PythonScriptPlugin/Private/PyWrapperObject.cpp
@@ -1707,9 +1707,9 @@ public:
 			if (InPyFuncDef->FuncRetType && InPyFuncDef->FuncRetType != Py_None)
 			{
 				// If we have a tuple, then we actually want to return a bool but add every type within the tuple as output parameters
-				const bool bOptionalReturn = static_cast<bool>(PyTuple_Check(InPyFuncDef->FuncRetType));
+				const bool bIsPackedReturn = PyTuple_Check(InPyFuncDef->FuncRetType) && !PyUtil::IsOptionalPropertyDef(InPyFuncDef->FuncRetType);
 
-				PyObject* RetType = bOptionalReturn ? (PyObject*)&PyBool_Type : InPyFuncDef->FuncRetType;
+				PyObject* RetType = bIsPackedReturn ? (PyObject*)&PyBool_Type : InPyFuncDef->FuncRetType;
 				FProperty* RetProp = PyUtil::CreateProperty(RetType, 1, Func, TEXT("ReturnValue"));
 				if (!RetProp)
 				{
@@ -1719,7 +1719,7 @@ public:
 				RetProp->PropertyFlags |= (CPF_Parm | CPF_OutParm | CPF_ReturnParm);
 				Func->AddCppProperty(RetProp);
 
-				if (bOptionalReturn)
+				if (bIsPackedReturn)
 				{
 					const int32 NumOutArgs = PyTuple_Size(InPyFuncDef->FuncRetType);
 					for (int32 ArgIndex = 0; ArgIndex < NumOutArgs; ++ArgIndex)
