Fixed FPyWrapperSet::Intersection/IntersectionUpdate failing for builds where TSet was TCompactSet

The updated implementation (iterating in reverse) works correctly for both TSparseSet and TCompactSet

#Synced-CL 50391842

diff --git a/Engine/Plugins/Experimental/PythonScriptPlugin/Source/PythonScriptPlugin/Private/PyWrapperSet.cpp b/Engine/Plugins/Experimental/PythonScriptPlugin/Source/PythonScriptPlugin/Private/PyWrapperSet.cpp
--- a/Engine/Plugins/Experimental/PythonScriptPlugin/Source/PythonScriptPlugin/Private/PyWrapperSet.cpp
+++ b/Engine/Plugins/Experimental/PythonScriptPlugin/Source/PythonScriptPlugin/Private/PyWrapperSet.cpp
@@ -693,8 +693,12 @@ int FPyWrapperSet::IntersectionUpdate(FPyWrapperSet* InSelf, PyObject* InOthers)
 
 		FScriptSetHelper OtherScriptSetHelper(Other->SetProp, Other->SetInstance);
 
+		// We iterate in reverse, as when TSet is TCompactSet (see UE_USE_COMPACT_SET_AS_DEFAULT), 
+		// it will perform a RemoveSwap, vs TSparseSet which will leave holes until compaction
+		// Reverse iteration works correctly for both implementations, whereas forward iteration 
+		// fails for TCompactSet, as it will skip over the element swapped from the end of the set
 		const int32 SelfSparseCount = SelfScriptSetHelper.GetMaxIndex();
-		for (int32 SelfSparseIndex = 0; SelfSparseIndex < SelfSparseCount; ++SelfSparseIndex)
+		for (int32 SelfSparseIndex = SelfSparseCount - 1; SelfSparseIndex >= 0; --SelfSparseIndex)
 		{
 			if (!SelfScriptSetHelper.IsValidIndex(SelfSparseIndex))
 			{
@@ -704,7 +708,6 @@ int FPyWrapperSet::IntersectionUpdate(FPyWrapperSet* InSelf, PyObject* InOthers)
 			const void* SelfElementPtr = SelfScriptSetHelper.GetElementPtr(SelfSparseIndex);
 			if (OtherScriptSetHelper.FindElementIndexFromHash(SelfElementPtr) == INDEX_NONE)
 			{
-				// This is safe as long as the set isn't compacted as we're iterating it
 				SelfScriptSetHelper.RemoveAt(SelfSparseIndex);
 			}
 		}
