Properly reset the enabled state of python when it is shutdown.

#Synced-CL 50460350

diff --git a/Engine/Plugins/Experimental/PythonScriptPlugin/Source/PythonScriptPlugin/Private/PythonScriptPlugin.cpp b/Engine/Plugins/Experimental/PythonScriptPlugin/Source/PythonScriptPlugin/Private/PythonScriptPlugin.cpp
--- a/Engine/Plugins/Experimental/PythonScriptPlugin/Source/PythonScriptPlugin/Private/PythonScriptPlugin.cpp
+++ b/Engine/Plugins/Experimental/PythonScriptPlugin/Source/PythonScriptPlugin/Private/PythonScriptPlugin.cpp
@@ -1533,6 +1533,7 @@ void FPythonScriptPlugin::ShutdownPython()
 	bIsFullyInitialized = false;
 	bIsForceEnabledAtRuntime = false;
 	bHasTicked = false;
+	bIsEnabled = false;
 }
 
 void FPythonScriptPlugin::RequestStubCodeGeneration()
