[Python]
- Remove GC during startup loading to avoid costly flush

#Synced-CL 50389147

diff --git a/Engine/Plugins/Experimental/PythonScriptPlugin/Source/PythonScriptPlugin/Private/PyWrapperTypeRegistry.cpp b/Engine/Plugins/Experimental/PythonScriptPlugin/Source/PythonScriptPlugin/Private/PyWrapperTypeRegistry.cpp
--- a/Engine/Plugins/Experimental/PythonScriptPlugin/Source/PythonScriptPlugin/Private/PyWrapperTypeRegistry.cpp
+++ b/Engine/Plugins/Experimental/PythonScriptPlugin/Source/PythonScriptPlugin/Private/PyWrapperTypeRegistry.cpp
@@ -828,7 +828,12 @@ void FPyWrapperTypeReinstancer::ProcessPending()
 		ClassesToReinstance.Reset();
 		StructsToReinstance.Reset();
 
-		CollectGarbage(GARBAGE_COLLECTION_KEEPFLAGS);
+		// Skip the GC in case we're actively loading to avoid a full flush especially during engine startup.
+		// The next GC will take care of the cleanup.
+		if (!IsAsyncLoading())
+		{
+			CollectGarbage(GARBAGE_COLLECTION_KEEPFLAGS);
+		}
 	}
 }
 
