Disable the pip installer when using Python in a restrictive environement.

#Synced-CL 50653521

diff --git a/Engine/Plugins/Experimental/PythonScriptPlugin/Source/PythonScriptPlugin/Private/PythonScriptPlugin.cpp b/Engine/Plugins/Experimental/PythonScriptPlugin/Source/PythonScriptPlugin/Private/PythonScriptPlugin.cpp
--- a/Engine/Plugins/Experimental/PythonScriptPlugin/Source/PythonScriptPlugin/Private/PythonScriptPlugin.cpp
+++ b/Engine/Plugins/Experimental/PythonScriptPlugin/Source/PythonScriptPlugin/Private/PythonScriptPlugin.cpp
@@ -1287,7 +1287,7 @@ bool FPythonScriptPlugin::InitializePython()
 	{
 		// Create the top-level "unreal" module
 		PyUnrealModule = FPyObjectPtr::NewReference(PyImport_AddModule("unreal"));
-		
+
 		// Import "unreal" into the console by default
 		PyDict_SetItemString(PyConsoleGlobalDict, "unreal", PyUnrealModule);
 
@@ -1312,17 +1312,23 @@ bool FPythonScriptPlugin::InitializePython()
 		FPyWrapperTypeRegistry::Get().OnModuleDirtied().AddRaw(this, &FPythonScriptPlugin::OnModuleDirtied);
 		FModuleManager::Get().OnModulesChanged().AddRaw(this, &FPythonScriptPlugin::OnModulesChanged);
 
+
+		UE::Tasks::FTask PipTask;
+		// We don't support pip installs in a restrictive environment at the moment.
+		if (!UE::Python::Private::bIsPythonInRestrictiveMode)
+		{
 #if WITH_EDITOR
-		IPipInstall& PipInstaller = IPipInstall::Get();
-		UE::Tasks::FTask PipTask = 
-			UE::Tasks::Launch(
-				TEXT("PipInstallAsync"),
-				[&PipInstaller]()
-				{
-					// Init PipInstall task
-					PipInstaller.InitPipInstall();
-				}
-			);
+			IPipInstall& PipInstaller = IPipInstall::Get();
+			PipTask =
+				UE::Tasks::Launch(
+					TEXT("PipInstallAsync"),
+					[&PipInstaller]()
+					{
+						// Init PipInstall task
+						PipInstaller.InitPipInstall();
+					}
+				);
+		}
 #endif // WITH_EDITOR
 
 		// We only allow that during boot
@@ -1332,22 +1338,26 @@ bool FPythonScriptPlugin::InitializePython()
 		FPyWrapperTypeRegistry::Get().GenerateWrappedTypes(bAllowMultithreadedGeneration);
 
 #if WITH_EDITOR
-		PipTask.Wait();
+		// We don't support pip installs in a restrictive environment at the moment.
+		if (!UE::Python::Private::bIsPythonInRestrictiveMode)
+		{
+			PipTask.Wait();
 
-		// Always add pip install site-package path since actual installs may be deferred
-		PipInstaller.RegisterPipSitePackagesPath();
+			IPipInstall& PipInstaller = IPipInstall::Get();
+			// Always add pip install site-package path since actual installs may be deferred
+			PipInstaller.RegisterPipSitePackagesPath();
 
-		// Check settings/cmd-line for whether pip installer is forced to run before initialization finished
-		bool bPipInstallOnInit = GetDefault<UPythonScriptPluginSettings>()->bRunPipInstallOnStartup;
-		bool bCmdLineForceInstall = FParse::Param(FCommandLine::Get(), TEXT("ForcePipInstallOnInit"));
-		bool bForceInstall = bPipInstallOnInit || bCmdLineForceInstall;
-		const int32 NumPkgs = PipInstaller.GetNumPackagesToInstall();
-		if (bForceInstall && NumPkgs > 0)
-		{
-			// Run PipInstall (in game thread) if any of the force options are on
-			PipInstallLauncher::StartSync(PipInstaller);
+			// Check settings/cmd-line for whether pip installer is forced to run before initialization finished
+			bool bPipInstallOnInit = GetDefault<UPythonScriptPluginSettings>()->bRunPipInstallOnStartup;
+			bool bCmdLineForceInstall = FParse::Param(FCommandLine::Get(), TEXT("ForcePipInstallOnInit"));
+			bool bForceInstall = bPipInstallOnInit || bCmdLineForceInstall;
+			const int32 NumPkgs = PipInstaller.GetNumPackagesToInstall();
+			if (bForceInstall && NumPkgs > 0)
+			{
+				// Run PipInstall (in game thread) if any of the force options are on
+				PipInstallLauncher::StartSync(PipInstaller);
+			}
 		}
-		
 #endif
 
 		// Initialize the tick handler
