Fixed FPyReferenceCollector reading the wrong address when collecting Optional properties

#rnx

#Synced-CL 50687180

diff --git a/Engine/Plugins/Experimental/PythonScriptPlugin/Source/PythonScriptPlugin/Private/PyReferenceCollector.cpp b/Engine/Plugins/Experimental/PythonScriptPlugin/Source/PythonScriptPlugin/Private/PyReferenceCollector.cpp
--- a/Engine/Plugins/Experimental/PythonScriptPlugin/Source/PythonScriptPlugin/Private/PyReferenceCollector.cpp
+++ b/Engine/Plugins/Experimental/PythonScriptPlugin/Source/PythonScriptPlugin/Private/PyReferenceCollector.cpp
@@ -461,9 +461,13 @@ void FPyReferenceCollector::AddReferencedObjectsFromPropertyInternal(FReferenceC
 
 	if (const FOptionalProperty* CastProp = CastField<FOptionalProperty>(InProp))
 	{
-		if (CastProp->IsSet(InBaseAddr))
+		for (int32 ArrIndex = 0; ArrIndex < InProp->ArrayDim; ++ArrIndex)
 		{
-			AddReferencedObjectsFromPropertyInternal(InCollector, CastProp->GetValueProperty(), CastProp->GetValuePointerForReadOrReplace(InBaseAddr), InFlags, OutValueChanged);
+			void* ValuePtr = CastProp->ContainerPtrToValuePtr<void>(InBaseAddr, ArrIndex);
+			if (CastProp->IsSet(ValuePtr))
+			{
+				AddReferencedObjectsFromPropertyInternal(InCollector, CastProp->GetValueProperty(), CastProp->GetValuePointerForReadOrReplace(ValuePtr), InFlags, OutValueChanged);
+			}
 		}
 		return;
 	}
