Made the recently executed python files menu not capture the files that are run on startup.

#Synced-CL 50972958

diff --git a/Engine/Plugins/Experimental/PythonScriptPlugin/Source/PythonScriptPlugin/Private/PythonScriptPlugin.cpp b/Engine/Plugins/Experimental/PythonScriptPlugin/Source/PythonScriptPlugin/Private/PythonScriptPlugin.cpp
--- a/Engine/Plugins/Experimental/PythonScriptPlugin/Source/PythonScriptPlugin/Private/PythonScriptPlugin.cpp
+++ b/Engine/Plugins/Experimental/PythonScriptPlugin/Source/PythonScriptPlugin/Private/PythonScriptPlugin.cpp
@@ -1600,7 +1600,6 @@ void FPythonScriptPlugin::RunStartupScripts()
 		return;
 	}
 
-	bIsFullyInitialized = true;
 
 	// Run start-up scripts now
 	TArray<FString> PySysPaths;
@@ -1693,6 +1692,8 @@ void FPythonScriptPlugin::RunStartupScripts()
 	}
 	DeferredCommands.Reset();
 
+	bIsFullyInitialized = true;
+
 	// Notify any external listeners
 	OnPythonInitializedDelegate.Broadcast();
 	OnPythonInitializedDelegate.Clear();
@@ -1896,7 +1897,8 @@ bool FPythonScriptPlugin::RunFile(const TCHAR* InFile, const TCHAR* InArgs, FPyt
 	FString FileStr;
 	bool bLoaded = FFileHelper::LoadFileToString(FileStr, *ResolvedFilePath);
 #if WITH_EDITOR
-	if (CmdMenu)
+	// We don't want to update the menu entries for the recent files when running the startups scripts
+	if (CmdMenu && bIsFullyInitialized)
 	{
 		CmdMenu->OnRunFile(ResolvedFilePath, bLoaded);
 	}
