Fix crash when removing individual items from uarrays in python (del some_array[0]).

#Synced-CL 51097783

diff --git a/Engine/Plugins/Experimental/PythonScriptPlugin/Source/PythonScriptPlugin/Private/PyWrapperArray.cpp b/Engine/Plugins/Experimental/PythonScriptPlugin/Source/PythonScriptPlugin/Private/PyWrapperArray.cpp
--- a/Engine/Plugins/Experimental/PythonScriptPlugin/Source/PythonScriptPlugin/Private/PyWrapperArray.cpp
+++ b/Engine/Plugins/Experimental/PythonScriptPlugin/Source/PythonScriptPlugin/Private/PyWrapperArray.cpp
@@ -438,6 +438,11 @@ int FPyWrapperArray::SetItem(FPyWrapperArray* InSelf, Py_ssize_t InIndex, PyObje
 		return -1;
 	}
 
+	if (InValue == nullptr)
+	{
+		return RemoveAt(InSelf, InIndex);
+	}
+
 	FScriptArrayHelper SelfScriptArrayHelper(InSelf->ArrayProp, InSelf->ArrayInstance);
 	const int32 ElementCount = SelfScriptArrayHelper.Num();
 	const Py_ssize_t ResolvedIndex = PyUtil::ResolveContainerIndexParam(InIndex, ElementCount);
@@ -457,6 +462,28 @@ int FPyWrapperArray::SetItem(FPyWrapperArray* InSelf, Py_ssize_t InIndex, PyObje
 	return 0;
 }
 
+int FPyWrapperArray::RemoveAt(FPyWrapperArray* InSelf, Py_ssize_t InIndex)
+{
+	if (!ValidateInternalState(InSelf))
+	{
+		return -1;
+	}
+
+	FScriptArrayHelper SelfScriptArrayHelper(InSelf->ArrayProp, InSelf->ArrayInstance);
+	const int32 ElementCount = SelfScriptArrayHelper.Num();
+	const Py_ssize_t ResolvedIndex = PyUtil::ResolveContainerIndexParam(InIndex, ElementCount);
+
+	const int ValidateIndexResult = PyUtil::ValidateContainerIndexParam(ResolvedIndex, ElementCount, InSelf->ArrayProp, *PyUtil::GetErrorContext(InSelf));
+	if (ValidateIndexResult != 0)
+	{
+		return ValidateIndexResult;
+	}
+
+	SelfScriptArrayHelper.RemoveValues(ResolvedIndex);
+
+	return 0;
+}
+
 int FPyWrapperArray::Contains(FPyWrapperArray* InSelf, PyObject* InValue)
 {
 	if (!ValidateInternalState(InSelf))
diff --git a/Engine/Plugins/Experimental/PythonScriptPlugin/Source/PythonScriptPlugin/Private/PyWrapperArray.h b/Engine/Plugins/Experimental/PythonScriptPlugin/Source/PythonScriptPlugin/Private/PyWrapperArray.h
--- a/Engine/Plugins/Experimental/PythonScriptPlugin/Source/PythonScriptPlugin/Private/PyWrapperArray.h
+++ b/Engine/Plugins/Experimental/PythonScriptPlugin/Source/PythonScriptPlugin/Private/PyWrapperArray.h
@@ -59,6 +59,9 @@ struct FPyWrapperArray : public FPyWrapperBase
 	/** Set the item at index N (equivalent to 'x[N] = v' in Python) */
 	static int SetItem(FPyWrapperArray* InSelf, Py_ssize_t InIndex, PyObject* InValue);
 
+	/** Removes the item at index N */
+	static int RemoveAt(FPyWrapperArray* InSelf, Py_ssize_t InIndex);
+
 	/** Does this container have an entry with the given value? (equivalent to 'v in x' in Python) */
 	static int Contains(FPyWrapperArray* InSelf, PyObject* InValue);
 
