Fixed abstract type hint errors in python wrappers for Array, Set, Map.

#Synced-CL 51097783

diff --git a/Engine/Plugins/Experimental/PythonScriptPlugin/Source/PythonScriptPlugin/Private/PyWrapperTypeRegistry.cpp b/Engine/Plugins/Experimental/PythonScriptPlugin/Source/PythonScriptPlugin/Private/PyWrapperTypeRegistry.cpp
--- a/Engine/Plugins/Experimental/PythonScriptPlugin/Source/PythonScriptPlugin/Private/PyWrapperTypeRegistry.cpp
+++ b/Engine/Plugins/Experimental/PythonScriptPlugin/Source/PythonScriptPlugin/Private/PyWrapperTypeRegistry.cpp
@@ -3092,7 +3092,7 @@ void FPyWrapperTypeRegistry::GenerateStubCodeForWrappedTypes(const EPyOnlineDocs
 
 		if (PyGenUtil::IsTypeHintingEnabled())
 		{
-			PythonScript.WriteLine(TEXT("from typing import Any, Callable, Dict, ItemsView, Iterable, Iterator, KeysView, List, Mapping, MutableMapping, MutableSequence, MutableSet, Optional, Set, Tuple, Type, TypeVar, Union, ValuesView"));
+			PythonScript.WriteLine(TEXT("from typing import Any, Callable, Dict, ItemsView, Iterable, Iterator, KeysView, List, Mapping, MutableMapping, MutableSequence, MutableSet, Optional, Set, Sequence, Tuple, Type, TypeVar, Union, ValuesView"));
 			PythonScript.WriteLine(TEXT("_T = TypeVar('_T')"));                 // For casting operations.
 			PythonScript.WriteLine(TEXT("_ElemType = TypeVar('_ElemType')"));   // For Array, FixedArray, Set classes along with casting operations.
 			PythonScript.WriteLine(TEXT("_KeyType = TypeVar('_KeyType')"));     // For Mapping[KV, KT] used by Map class.
@@ -3362,11 +3362,16 @@ void FPyWrapperTypeRegistry::GenerateStubCodeForWrappedType(PyTypeObject* PyType
 			: PyGenUtil::EPythonizeFlags::WithTypeHinting | PyGenUtil::EPythonizeFlags::WithOptionalType) // For type checker.
 		: PyGenUtil::EPythonizeFlags::DefaultConstructStructs | PyGenUtil::EPythonizeFlags::DefaultConstructDateTime | PyGenUtil::EPythonizeFlags::UseStrictTyping; // hinting is off.
 
-	if (PyGenUtil::IsTypeHintingEnabled() &&  (PyTypeName == TEXT("Array") || PyTypeName == TEXT("FixedArray")))
+	if (PyGenUtil::IsTypeHintingEnabled() &&  PyTypeName == TEXT("Array"))
 	{
 		// Implies that _ElemType = TypeVar('_ElemType') is already written to the file.
 		OutPythonScript.WriteLine(FString::Printf(TEXT("class %s(%s, MutableSequence[_ElemType]):"), *PyTypeName, UTF8_TO_TCHAR(PyType->tp_base->tp_name)));
 	}
+	else if (PyGenUtil::IsTypeHintingEnabled() && PyTypeName == TEXT("FixedArray"))
+	{
+		// Implies that _ElemType = TypeVar('_ElemType') is already written to the file.
+		OutPythonScript.WriteLine(FString::Printf(TEXT("class %s(%s, Sequence[_ElemType]):"), *PyTypeName, UTF8_TO_TCHAR(PyType->tp_base->tp_name)));
+	}
 	else if (PyGenUtil::IsTypeHintingEnabled() && PyTypeName == TEXT("Set"))
 	{
 		// Implies that _ElemType = TypeVar('_ElemType') is already written to the file.
@@ -3806,6 +3811,11 @@ void FPyWrapperTypeRegistry::GenerateStubCodeForWrappedType(PyTypeObject* PyType
 
 				ExportNoBodyMethod(TEXT("def __setitem__(self, index: int, value: _ElemType) -> None:"));
 				ExportNoBodyMethod(TEXT("def __getitem__(self, index: int) -> _ElemType:"));
+				ExportNoBodyMethod(TEXT("def __len__(self) -> int:"));
+				if (PyType == &PyWrapperArrayType)
+				{
+					ExportNoBodyMethod(TEXT("def __delitem__(self, key: Union[int, slice, Any]) -> None:"));
+				}
 			}
 			else if (PyType == &PyWrapperSetType)
 			{
@@ -3813,6 +3823,10 @@ void FPyWrapperTypeRegistry::GenerateStubCodeForWrappedType(PyTypeObject* PyType
 				ExportNoBodyMethod(TEXT("def __init__(self, type: type) -> None:"));
 				bWriteDefaultInit = false;
 				bHasExportedClassData = true;
+
+				ExportNoBodyMethod(TEXT("def __len__(self) -> int:"));
+				ExportNoBodyMethod(TEXT("def __contains__(self, x: object) -> bool:"));
+				ExportNoBodyMethod(TEXT("def __iter__(self) -> Iterator[Any]:"));
 			}
 			else if (PyType == &PyWrapperMapType)
 			{
@@ -3822,6 +3836,9 @@ void FPyWrapperTypeRegistry::GenerateStubCodeForWrappedType(PyTypeObject* PyType
 
 				ExportNoBodyMethod(TEXT("def __setitem__(self, key: _KeyType, value: _ValueType) -> None:"));
 				ExportNoBodyMethod(TEXT("def __getitem__(self, key: _KeyType) -> _ValueType:"));
+				ExportNoBodyMethod(TEXT("def __len__(self) -> int:"));
+				ExportNoBodyMethod(TEXT("def __iter__(self) -> Iterator[Any]:"));
+				ExportNoBodyMethod(TEXT("def __delitem__(self, key: Union[int, slice, Any]) -> None:"));
 			}
 			else if (PyType == &PyWrapperNameType)
 			{
