[40.20] Adding Python Error stacktrace output to the log for the experimental Python Plugin

#Synced-CL 51207507

diff --git a/Engine/Plugins/Experimental/PythonScriptPlugin/Source/PythonScriptPlugin/Private/PythonScriptPlugin.cpp b/Engine/Plugins/Experimental/PythonScriptPlugin/Source/PythonScriptPlugin/Private/PythonScriptPlugin.cpp
--- a/Engine/Plugins/Experimental/PythonScriptPlugin/Source/PythonScriptPlugin/Private/PythonScriptPlugin.cpp
+++ b/Engine/Plugins/Experimental/PythonScriptPlugin/Source/PythonScriptPlugin/Private/PythonScriptPlugin.cpp
@@ -1950,6 +1950,13 @@ bool FPythonScriptPlugin::RunFile(const TCHAR* InFile, const TCHAR* InArgs, FPyt
 		}
 		else if (PyUtil::LogPythonError(&InOutPythonCommand.CommandResult))
 		{
+			TArray<FString> ErrorLines;
+			InOutPythonCommand.CommandResult.ParseIntoArrayLines(ErrorLines);
+
+			for (const FString& ErrorLine : ErrorLines)
+			{
+				InOutPythonCommand.LogOutput.Add(FPythonLogOutputEntry{ EPythonLogOutputType::Error, ErrorLine });
+			}
 			return false;
 		}
 	}
