Added tracing to Python garbage collection

#Synced-CL 51631837

diff --git a/Engine/Plugins/Experimental/PythonScriptPlugin/Source/PythonScriptPlugin/Private/PyUtil.cpp b/Engine/Plugins/Experimental/PythonScriptPlugin/Source/PythonScriptPlugin/Private/PyUtil.cpp
--- a/Engine/Plugins/Experimental/PythonScriptPlugin/Source/PythonScriptPlugin/Private/PyUtil.cpp
+++ b/Engine/Plugins/Experimental/PythonScriptPlugin/Source/PythonScriptPlugin/Private/PyUtil.cpp
@@ -1797,6 +1797,7 @@ bool ReThrowPythonError(FString* OutError)
 
 void CollectGarbage(int32 Generations)
 {
+	TRACE_CPUPROFILER_EVENT_SCOPE(PyUtil::CollectGarbage);
 	if (Generations == 2)
 	{
 		// The C API hardcodes Generations == 2 
