Fixed get/set_editor_property/properties not being able to find properties on a UserDefinedStruct via their clean name

#rnx
#Synced-CL 51753475

#Synced-CL 51753475

diff --git a/Engine/Plugins/Experimental/PythonScriptPlugin/Source/PythonScriptPlugin/Private/PyWrapperStruct.cpp b/Engine/Plugins/Experimental/PythonScriptPlugin/Source/PythonScriptPlugin/Private/PyWrapperStruct.cpp
--- a/Engine/Plugins/Experimental/PythonScriptPlugin/Source/PythonScriptPlugin/Private/PyWrapperStruct.cpp
+++ b/Engine/Plugins/Experimental/PythonScriptPlugin/Source/PythonScriptPlugin/Private/PyWrapperStruct.cpp
@@ -1027,6 +1027,10 @@ PyTypeObject InitializePyWrapperStructType()
 			const FName ResolvedName = FPyWrapperStructMetaData::ResolvePropertyName(InSelf, InPythonPropName);
 			const FProperty* ResolvedProp = InSelf->ScriptStruct->FindPropertyByName(ResolvedName);
 			if (!ResolvedProp)
+			{
+				ResolvedProp = InSelf->ScriptStruct->CustomFindProperty(ResolvedName);
+			}
+			if (!ResolvedProp)
 			{
 				PyUtil::SetPythonError(PyExc_Exception, InSelf, *FString::Printf(TEXT("Failed to find property '%s' for attribute '%s' on '%s'"), *ResolvedName.ToString(), *InPythonPropName.ToString(), *InSelf->ScriptStruct->GetName()));
 				return WrappedPropDef;
