Python reinstancing fixes

* Fixed a dangling PyMetaData reference after a Python generated type has been destroyed.
* Fixed a ref-count mismatch on the PyType when an existing PyWrapperObject has its instance replaced.
* Fixed a crash during CDO destruction when a class marked CLASS_NewerVersionExists had its base class property links destroyed.

#Synced-CL 52288889

diff --git a/Engine/Plugins/Experimental/PythonScriptPlugin/Source/PythonScriptPlugin/Private/PyWrapperEnum.cpp b/Engine/Plugins/Experimental/PythonScriptPlugin/Source/PythonScriptPlugin/Private/PyWrapperEnum.cpp
--- a/Engine/Plugins/Experimental/PythonScriptPlugin/Source/PythonScriptPlugin/Private/PyWrapperEnum.cpp
+++ b/Engine/Plugins/Experimental/PythonScriptPlugin/Source/PythonScriptPlugin/Private/PyWrapperEnum.cpp
@@ -950,6 +950,7 @@ void UPythonGeneratedEnum::ReleasePythonResources()
 	{
 		FPyScopedGIL GIL;
 		UnregisterGeneratedType();
+		FPyWrapperEnumMetaData::SetMetaData(PyType.Get(), nullptr);
 		PyType.Reset();
 	}
 	else
diff --git a/Engine/Plugins/Experimental/PythonScriptPlugin/Source/PythonScriptPlugin/Private/PyWrapperObject.cpp b/Engine/Plugins/Experimental/PythonScriptPlugin/Source/PythonScriptPlugin/Private/PyWrapperObject.cpp
--- a/Engine/Plugins/Experimental/PythonScriptPlugin/Source/PythonScriptPlugin/Private/PyWrapperObject.cpp
+++ b/Engine/Plugins/Experimental/PythonScriptPlugin/Source/PythonScriptPlugin/Private/PyWrapperObject.cpp
@@ -1185,7 +1185,10 @@ void FPyWrapperObjectMetaData::AddInstanceReferencedObjects(FPyWrapperBase* Inst
 		PyTypeObject* NewPyType = FPyWrapperTypeRegistry::Get().GetWrappedClassType(Self->ObjectInstance->GetClass());
 		if (PyType_IsSubtype(NewPyType, &PyWrapperObjectType) && NewPyType->tp_basicsize == Py_TYPE(Self)->tp_basicsize)
 		{
-			Py_SET_TYPE(Self, NewPyType); // todo: is this safe?
+			PyTypeObject* OldPyType = Py_TYPE(Self);
+			Py_INCREF(NewPyType);   // balance the future Py_DECREF in subtype_dealloc
+			Py_SET_TYPE(Self, NewPyType);
+			Py_DECREF(OldPyType);   // balance the Py_INCREF from tp_alloc
 		}
 		else
 		{
@@ -2063,6 +2066,7 @@ void UPythonGeneratedClass::ReleasePythonResources()
 	{
 		FPyScopedGIL GIL;
 		UnregisterGeneratedType();
+		FPyWrapperObjectMetaData::SetMetaData(PyType.Get(), nullptr);
 		PyType.Reset();
 		PyPostInitFunction.Reset();
 		for (const TSharedPtr<PyGenUtil::FFunctionDef>& FunctionDef : FunctionDefs)
diff --git a/Engine/Plugins/Experimental/PythonScriptPlugin/Source/PythonScriptPlugin/Private/PyWrapperStruct.cpp b/Engine/Plugins/Experimental/PythonScriptPlugin/Source/PythonScriptPlugin/Private/PyWrapperStruct.cpp
--- a/Engine/Plugins/Experimental/PythonScriptPlugin/Source/PythonScriptPlugin/Private/PyWrapperStruct.cpp
+++ b/Engine/Plugins/Experimental/PythonScriptPlugin/Source/PythonScriptPlugin/Private/PyWrapperStruct.cpp
@@ -1724,6 +1724,7 @@ void UPythonGeneratedStruct::ReleasePythonResources()
 	{
 		FPyScopedGIL GIL;
 		UnregisterGeneratedType();
+		FPyWrapperStructMetaData::SetMetaData(PyType.Get(), nullptr);
 		PyType.Reset();
 		PyPostInitFunction.Reset();
 	}
diff --git a/Engine/Plugins/Experimental/PythonScriptPlugin/Source/PythonScriptPlugin/Private/PyWrapperTypeRegistry.cpp b/Engine/Plugins/Experimental/PythonScriptPlugin/Source/PythonScriptPlugin/Private/PyWrapperTypeRegistry.cpp
--- a/Engine/Plugins/Experimental/PythonScriptPlugin/Source/PythonScriptPlugin/Private/PyWrapperTypeRegistry.cpp
+++ b/Engine/Plugins/Experimental/PythonScriptPlugin/Source/PythonScriptPlugin/Private/PyWrapperTypeRegistry.cpp
@@ -787,34 +787,51 @@ FPyWrapperTypeReinstancer& FPyWrapperTypeReinstancer::Get()
 
 void FPyWrapperTypeReinstancer::AddPendingClass(UPythonGeneratedClass* OldClass, UPythonGeneratedClass* NewClass)
 {
-	ClassesToReinstance.Emplace(MakeTuple(OldClass, NewClass));
+	ClassesToReinstance.Emplace(OldClass, NewClass);
 }
 
 void FPyWrapperTypeReinstancer::AddPendingStruct(UPythonGeneratedStruct* OldStruct, UPythonGeneratedStruct* NewStruct)
 {
-	StructsToReinstance.Emplace(MakeTuple(OldStruct, NewStruct));
+	StructsToReinstance.Emplace(OldStruct, NewStruct);
 }
 
 void FPyWrapperTypeReinstancer::ProcessPending()
 {
 	if (ClassesToReinstance.Num() > 0 || StructsToReinstance.Num() > 0)
 	{
-		TUniquePtr<FReload> Reload(new FReload(EActiveReloadType::Reinstancing, TEXT(""), *GLog));
+		TUniquePtr<FReload> Reload = MakeUnique<FReload>(EActiveReloadType::Reinstancing, TEXT(""), *GLog);
 
-		for (const auto& ClassToReinstancePair : ClassesToReinstance)
+		for (const TTuple<TObjectPtr<UPythonGeneratedClass>, TObjectPtr<UPythonGeneratedClass>>& ClassToReinstancePair : ClassesToReinstance)
 		{
 			if (ClassToReinstancePair.Key && ClassToReinstancePair.Value)
 			{
-				if (!ClassToReinstancePair.Value->HasAnyClassFlags(CLASS_NewerVersionExists))
+				// Resolve any CLASS_NewerVersionExists chains, so that any older classes all get reinstanced onto the most recent class
+				UPythonGeneratedClass* NewClass = ClassToReinstancePair.Value;
+				while (NewClass && NewClass->HasAnyClassFlags(CLASS_NewerVersionExists))
+				{
+					// We expect to find the replacement for NewClass in ClassesToReinstance, since 
+					// that's the only way it should have been marked as CLASS_NewerVersionExists
+					NewClass = ClassesToReinstance.FindRef(NewClass);
+					ensureAlways(NewClass);
+				}
+
+				if (NewClass)
 				{
 					// Assume the classes have changed
-					Reload->NotifyChange(ClassToReinstancePair.Value, ClassToReinstancePair.Key);
+					Reload->NotifyChange(NewClass, ClassToReinstancePair.Key);
+				}
+				else
+				{
+					// We failed to resolve the new class; old class won't be purged via Reinstance(),
+					// but its property links must still be cleared so that DestroyNonNativeProperties
+					// on its CDO doesn't walk FFields that may be freed by the base class FinishDestroy
+					ClassToReinstancePair.Key->DestroyChildPropertiesAndResetPropertyLinks();
 				}
 			}
 		}
 
 		// This doesn't do anything ATM
-		for (const auto& StructToReinstancePair : StructsToReinstance)
+		for (const TTuple<TObjectPtr<UPythonGeneratedStruct>, TObjectPtr<UPythonGeneratedStruct>>& StructToReinstancePair : StructsToReinstance)
 		{
 			if (StructToReinstancePair.Key && StructToReinstancePair.Value)
 			{
@@ -845,17 +862,8 @@ void FPyWrapperTypeReinstancer::Shutdown()
 
 void FPyWrapperTypeReinstancer::AddReferencedObjects(FReferenceCollector& InCollector)
 {
-	for (auto& ClassToReinstancePair : ClassesToReinstance)
-	{
-		InCollector.AddReferencedObject(ClassToReinstancePair.Key);
-		InCollector.AddReferencedObject(ClassToReinstancePair.Value);
-	}
-
-	for (auto& StructToReinstancePair : StructsToReinstance)
-	{
-		InCollector.AddReferencedObject(StructToReinstancePair.Key);
-		InCollector.AddReferencedObject(StructToReinstancePair.Value);
-	}
+	InCollector.AddReferencedObjects(ClassesToReinstance);
+	InCollector.AddReferencedObjects(StructsToReinstance);
 }
 
 FString FPyWrapperTypeReinstancer::GetReferencerName() const
diff --git a/Engine/Plugins/Experimental/PythonScriptPlugin/Source/PythonScriptPlugin/Private/PyWrapperTypeRegistry.h b/Engine/Plugins/Experimental/PythonScriptPlugin/Source/PythonScriptPlugin/Private/PyWrapperTypeRegistry.h
--- a/Engine/Plugins/Experimental/PythonScriptPlugin/Source/PythonScriptPlugin/Private/PyWrapperTypeRegistry.h
+++ b/Engine/Plugins/Experimental/PythonScriptPlugin/Source/PythonScriptPlugin/Private/PyWrapperTypeRegistry.h
@@ -348,10 +348,10 @@ public:
 
 private:
 	/** Pending pairs of classes that to be re-instanced */
-	TArray<TPair<TObjectPtr<UPythonGeneratedClass>, TObjectPtr<UPythonGeneratedClass>>> ClassesToReinstance;
+	TMap<TObjectPtr<UPythonGeneratedClass>, TObjectPtr<UPythonGeneratedClass>> ClassesToReinstance;
 
 	/** Pending pairs of structs that to be re-instanced */
-	TArray<TPair<TObjectPtr<UPythonGeneratedStruct>, TObjectPtr<UPythonGeneratedStruct>>> StructsToReinstance;
+	TMap<TObjectPtr<UPythonGeneratedStruct>, TObjectPtr<UPythonGeneratedStruct>> StructsToReinstance;
 };
 
 /** Singleton instance that maps Unreal types to Python types */
