Python-defined function return value fixes

Fixed the order that packed return values are built, so they match the order that the types are defined (and match C++). This was a bug previously fixed for the input arguments, but not the output arguments.

#Synced-CL 52671426

diff --git a/Engine/Plugins/Experimental/PythonScriptPlugin/Source/PythonScriptPlugin/Private/PyWrapperObject.cpp b/Engine/Plugins/Experimental/PythonScriptPlugin/Source/PythonScriptPlugin/Private/PyWrapperObject.cpp
--- a/Engine/Plugins/Experimental/PythonScriptPlugin/Source/PythonScriptPlugin/Private/PyWrapperObject.cpp
+++ b/Engine/Plugins/Experimental/PythonScriptPlugin/Source/PythonScriptPlugin/Private/PyWrapperObject.cpp
@@ -1724,8 +1724,9 @@ public:
 
 				if (bIsPackedReturn)
 				{
+					// Adding properties to a function inserts them into a linked list, so we loop backwards to get the order right
 					const int32 NumOutArgs = PyTuple_Size(InPyFuncDef->FuncRetType);
-					for (int32 ArgIndex = 0; ArgIndex < NumOutArgs; ++ArgIndex)
+					for (int32 ArgIndex = NumOutArgs - 1; ArgIndex >= 0; --ArgIndex)
 					{
 						PyObject* ArgTypeObj = PySequence_GetItem(InPyFuncDef->FuncRetType, ArgIndex);
 						FProperty* ArgProp = PyUtil::CreateProperty(ArgTypeObj, 1, Func, *FString::Printf(TEXT("OutValue%d"), ArgIndex));
@@ -1742,21 +1743,17 @@ public:
 			}
 
 			// Adding properties to a function inserts them into a linked list, so we loop backwards to get the order right
+			for (int32 ArgIndex = FuncArgNames.Num() - 1; ArgIndex >= 0; --ArgIndex)
 			{
-				int32 ArgIndex = FuncArgNames.Num() - 1;
-				while (ArgIndex >= 0)
+				PyObject* ArgTypeObj = PySequence_GetItem(InPyFuncDef->FuncParamTypes, ArgIndex);
+				FProperty* ArgProp = PyUtil::CreateProperty(ArgTypeObj, 1, Func, *FuncArgNames[ArgIndex]);
+				if (!ArgProp)
 				{
-					PyObject* ArgTypeObj = PySequence_GetItem(InPyFuncDef->FuncParamTypes, ArgIndex);
-					FProperty* ArgProp = PyUtil::CreateProperty(ArgTypeObj, 1, Func, *FuncArgNames[ArgIndex]);
-					if (!ArgProp)
-					{
-						PyUtil::SetPythonError(PyExc_Exception, PyType, *FString::Printf(TEXT("Failed to create property (%s) for function '%s' argument '%s'"), *PyUtil::GetFriendlyTypename(ArgTypeObj), *InFieldName, *FuncArgNames[ArgIndex]));
-						return false;
-					}
-					ArgProp->PropertyFlags |= CPF_Parm;
-					Func->AddCppProperty(ArgProp);
-					ArgIndex--;
+					PyUtil::SetPythonError(PyExc_Exception, PyType, *FString::Printf(TEXT("Failed to create property (%s) for function '%s' argument '%s'"), *PyUtil::GetFriendlyTypename(ArgTypeObj), *InFieldName, *FuncArgNames[ArgIndex]));
+					return false;
 				}
+				ArgProp->PropertyFlags |= CPF_Parm;
+				Func->AddCppProperty(ArgProp);
 			}
 		}
 		// Apply the defaults to the function arguments and build the Python method params
